ECABridge experimental MCP plugin

#Synced-CL 51504585

diff --git a/Engine/Plugins/Experimental/ECABridge/ECABridge.uplugin b/Engine/Plugins/Experimental/ECABridge/ECABridge.uplugin
new file mode 100644
--- /dev/null
+++ b/Engine/Plugins/Experimental/ECABridge/ECABridge.uplugin
@@ -0,0 +1,55 @@
+{
+	"FileVersion": 3,
+	"Version": 1,
+	"VersionName": "1.0",
+	"FriendlyName": "ECA Bridge",
+	"Description": "Epic Code Assistant integration for Unreal Editor - enables AI-assisted development through MCP protocol",
+	"Category": "Editor",
+	"CreatedBy": "Epic Games",
+	"CreatedByURL": "https://github.com/epicgames/epic-code-assistant",
+	"DocsURL": "",
+	"MarketplaceURL": "",
+	"SupportURL": "",
+	"CanContainContent": false,
+	"IsBetaVersion": true,
+	"IsExperimentalVersion": true,
+	"Installed": false,
+	"EnabledByDefault": false,
+	"Modules": [
+		{
+			"Name": "ECABridge",
+			"Type": "Editor",
+			"LoadingPhase": "Default"
+		}
+	],
+	"Plugins": [
+		{
+			"Name": "Niagara",
+			"Enabled": true
+		},
+		{
+			"Name": "EditorScriptingUtilities",
+			"Enabled": true
+		},
+		{
+			"Name": "Metasound",
+			"Enabled": true
+		},
+		{
+			"Name": "GeometryScripting",
+			"Enabled": true
+		},
+		{
+			"Name": "GeometryProcessing",
+			"Enabled": true
+		},
+		{
+			"Name": "MeshModelingToolset",
+			"Enabled": true
+		},
+		{
+			"Name": "ProceduralMeshComponent",
+			"Enabled": true
+		}
+	]
+}
diff --git a/Engine/Plugins/Experimental/ECABridge/Source/ECABridge/ECABridge.Build.cs b/Engine/Plugins/Experimental/ECABridge/Source/ECABridge/ECABridge.Build.cs
new file mode 100644
--- /dev/null
+++ b/Engine/Plugins/Experimental/ECABridge/Source/ECABridge/ECABridge.Build.cs
@@ -0,0 +1,77 @@
+// Copyright Epic Games, Inc. All Rights Reserved.
+
+using System.IO;
+using UnrealBuildTool;
+
+public class ECABridge : ModuleRules
+{
+	public ECABridge(ReadOnlyTargetRules Target) : base(Target)
+	{
+		PCHUsage = PCHUsageMode.UseExplicitOrSharedPCHs;
+		CppStandard = CppStandardVersion.Cpp23;
+
+		PublicDependencyModuleNames.AddRange(new string[]
+		{
+			"Core", "CoreUObject", "Engine", "EngineSettings", "Json", "JsonUtilities",
+			"HTTPServer",  // For MCP HTTP/SSE server
+			"Sockets",     // For ISocketSubsystem
+			"Networking"   // For networking utilities
+		});
+
+		PrivateDependencyModuleNames.AddRange(new string[]
+		{
+			// Editor core
+			"UnrealEd",
+			
+			// MetaSound
+			"MetasoundEngine", "MetasoundEditor", "MetasoundFrontend", "MetasoundGraphCore", "MetasoundStandardNodes", "AudioExtensions",
+			
+			// Mesh import/manipulation
+			"MeshDescription", "StaticMeshDescription",
+			// GeometryScript - Modern procedural mesh API
+			"GeometryCore", "GeometryFramework", "GeometryScriptingCore", "DynamicMesh", "MeshConversion",
+			// Modeling tools subdivision
+			"ModelingComponentsEditorOnly", "ModelingOperators",
+			// ProceduralMeshComponent - Runtime procedural meshes
+			"ProceduralMeshComponent",
+			
+			// Niagara FX
+			"Niagara", "NiagaraCore", "NiagaraEditor", "NiagaraShader",
+			"EditorSubsystem", "EditorFramework", "LevelEditor",
+			
+			// UI/Slate
+			"Slate", "SlateCore", "UMG", "UMGEditor", "PropertyEditor", "InputCore",
+			// MVVM ViewModel binding
+			"ModelViewViewModel", "ModelViewViewModelBlueprint", "ModelViewViewModelEditor",
+			
+			// Blueprint manipulation
+			"BlueprintGraph", "Kismet", "KismetCompiler", "GraphEditor",
+			
+			// Asset management
+			"AssetTools", "AssetRegistry", "ContentBrowser", "EditorScriptingUtilities",  // For EditorAssetLibrary
+			"MaterialEditor",            // For material editing utilities
+			"DataTableEditor",           // For DataTable editor UI refresh (BroadcastPostChange)
+			
+			// Additional utilities
+			"Projects", "DeveloperSettings", "ToolMenus",
+			
+			// Image processing
+			"ImageWrapper",
+			
+			// Interchange framework for FBX import
+			"InterchangeCore",
+			"InterchangeEngine",
+			
+			// Rendering (for thumbnail generation)
+			"RenderCore",
+			"RHI"
+		});
+
+		bUseRTTI = true;
+
+		// ModelViewViewModelBlueprint include path - UBT doesn't resolve it via PrivateDependencyModuleNames
+		// when the plugin is EnabledByDefault=false, so add it explicitly.
+		string MVVMBlueprintPath = System.IO.Path.Combine(EngineDirectory, "Plugins", "Runtime", "ModelViewViewModel", "Source", "ModelViewViewModelBlueprint", "Public");
+		PrivateIncludePaths.Add(MVVMBlueprintPath);
+	}
+}
diff --git a/Engine/Plugins/Experimental/ECABridge/Source/ECABridge/Private/BlueprintLisp.cpp b/Engine/Plugins/Experimental/ECABridge/Source/ECABridge/Private/BlueprintLisp.cpp
new file mode 100644
--- /dev/null
+++ b/Engine/Plugins/Experimental/ECABridge/Source/ECABridge/Private/BlueprintLisp.cpp
@@ -0,0 +1,936 @@
+// Copyright Epic Games, Inc. All Rights Reserved.
+
+#include "BlueprintLisp.h"
+#include "Serialization/JsonSerializer.h"
+#include "Serialization/JsonWriter.h"
+
+//------------------------------------------------------------------------------
+// FLispNode Implementation
+//------------------------------------------------------------------------------
+
+FLispNodePtr FLispNode::MakeNil()
+{
+	return MakeShared<FLispNode>();
+}
+
+FLispNodePtr FLispNode::MakeSymbol(const FString& Name)
+{
+	auto Node = MakeShared<FLispNode>();
+	Node->Type = ELispNodeType::Symbol;
+	Node->StringValue = Name;
+	return Node;
+}
+
+FLispNodePtr FLispNode::MakeKeyword(const FString& Name)
+{
+	auto Node = MakeShared<FLispNode>();
+	Node->Type = ELispNodeType::Keyword;
+	Node->StringValue = Name;
+	return Node;
+}
+
+FLispNodePtr FLispNode::MakeNumber(double Value)
+{
+	auto Node = MakeShared<FLispNode>();
+	Node->Type = ELispNodeType::Number;
+	Node->NumberValue = Value;
+	return Node;
+}
+
+FLispNodePtr FLispNode::MakeString(const FString& Value)
+{
+	auto Node = MakeShared<FLispNode>();
+	Node->Type = ELispNodeType::String;
+	Node->StringValue = Value;
+	return Node;
+}
+
+FLispNodePtr FLispNode::MakeList(const TArray<FLispNodePtr>& Items)
+{
+	auto Node = MakeShared<FLispNode>();
+	Node->Type = ELispNodeType::List;
+	Node->Children = Items;
+	return Node;
+}
+
+bool FLispNode::IsForm(const FString& FormName) const
+{
+	if (!IsList() || Children.Num() == 0)
+	{
+		return false;
+	}
+	
+	const auto& First = Children[0];
+	return First.IsValid() && First->IsSymbol() && First->StringValue.Equals(FormName, ESearchCase::IgnoreCase);
+}
+
+FString FLispNode::GetFormName() const
+{
+	if (IsList() && Children.Num() > 0 && Children[0].IsValid() && Children[0]->IsSymbol())
+	{
+		return Children[0]->StringValue;
+	}
+	return TEXT("");
+}
+
+FLispNodePtr FLispNode::operator[](int32 Index) const
+{
+	return Get(Index);
+}
+
+FLispNodePtr FLispNode::Get(int32 Index) const
+{
+	if (Index >= 0 && Index < Children.Num())
+	{
+		return Children[Index];
+	}
+	return MakeNil();
+}
+
+FLispNodePtr FLispNode::GetKeywordArg(const FString& Keyword) const
+{
+	if (!IsList())
+	{
+		return MakeNil();
+	}
+	
+	FString KeywordWithColon = Keyword.StartsWith(TEXT(":")) ? Keyword : (TEXT(":") + Keyword);
+	
+	for (int32 i = 0; i < Children.Num() - 1; i++)
+	{
+		if (Children[i].IsValid() && Children[i]->IsKeyword() && 
+			Children[i]->StringValue.Equals(KeywordWithColon, ESearchCase::IgnoreCase))
+		{
+			return Children[i + 1];
+		}
+	}
+	
+	return MakeNil();
+}
+
+bool FLispNode::HasKeyword(const FString& Keyword) const
+{
+	if (!IsList())
+	{
+		return false;
+	}
+	
+	FString KeywordWithColon = Keyword.StartsWith(TEXT(":")) ? Keyword : (TEXT(":") + Keyword);
+	
+	for (const auto& Child : Children)
+	{
+		if (Child.IsValid() && Child->IsKeyword() && 
+			Child->StringValue.Equals(KeywordWithColon, ESearchCase::IgnoreCase))
+		{
+			return true;
+		}
+	}
+	
+	return false;
+}
+
+FString FLispNode::ToString(bool bPretty, int32 IndentLevel) const
+{
+	FString Indent = bPretty ? FString::ChrN(IndentLevel * 2, ' ') : TEXT("");
+	FString Newline = bPretty ? TEXT("\n") : TEXT("");
+	FString Space = bPretty ? TEXT(" ") : TEXT(" ");
+	
+	switch (Type)
+	{
+	case ELispNodeType::Nil:
+		return TEXT("nil");
+		
+	case ELispNodeType::Symbol:
+		return StringValue;
+		
+	case ELispNodeType::Keyword:
+		return StringValue; // Already includes ':'
+		
+	case ELispNodeType::Number:
+		if (FMath::IsNearlyEqual(NumberValue, FMath::RoundToDouble(NumberValue)))
+		{
+			return FString::Printf(TEXT("%d"), (int64)NumberValue);
+		}
+		return FString::SanitizeFloat(NumberValue);
+		
+	case ELispNodeType::String:
+		{
+			FString Escaped = StringValue;
+			Escaped.ReplaceInline(TEXT("\\"), TEXT("\\\\"));
+			Escaped.ReplaceInline(TEXT("\""), TEXT("\\\""));
+			Escaped.ReplaceInline(TEXT("\n"), TEXT("\\n"));
+			Escaped.ReplaceInline(TEXT("\t"), TEXT("\\t"));
+			return FString::Printf(TEXT("\"%s\""), *Escaped);
+		}
+		
+	case ELispNodeType::List:
+		{
+			if (Children.Num() == 0)
+			{
+				return TEXT("()");
+			}
+			
+			// Check if this is a "complex" form that should be multi-line
+			bool bMultiLine = bPretty && Children.Num() > 3;
+			if (!bMultiLine && bPretty)
+			{
+				// Also go multi-line if any child is a list
+				for (const auto& Child : Children)
+				{
+					if (Child.IsValid() && Child->IsList() && Child->Children.Num() > 2)
+					{
+						bMultiLine = true;
+						break;
+					}
+				}
+			}
+			
+			// Check for special forms that should always be multi-line
+			FString FormName = GetFormName();
+			if (bPretty && (FormName == TEXT("event") || FormName == TEXT("func") || 
+				FormName == TEXT("macro") || FormName == TEXT("seq") ||
+				FormName == TEXT("branch") || FormName == TEXT("foreach") ||
+				FormName == TEXT("switch")))
+			{
+				bMultiLine = true;
+			}
+			
+			FString Result = TEXT("(");
+			
+			if (bMultiLine)
+			{
+				for (int32 i = 0; i < Children.Num(); i++)
+				{
+					if (i == 0)
+					{
+						// First element (form name) on same line
+						Result += Children[i].IsValid() ? Children[i]->ToString(false, 0) : TEXT("nil");
+					}
+					else if (Children[i].IsValid() && Children[i]->IsKeyword())
+					{
+						// Keyword on new line
+						Result += Newline + Indent + TEXT("  ");
+						Result += Children[i]->ToString(false, 0);
+						// Value follows keyword on same line or next depending on complexity
+						if (i + 1 < Children.Num())
+						{
+							i++;
+							if (Children[i].IsValid() && Children[i]->IsList() && Children[i]->Children.Num() > 2)
+							{
+								Result += Newline + Indent + TEXT("    ");
+								Result += Children[i]->ToString(bPretty, IndentLevel + 2);
+							}
+							else
+							{
+								Result += TEXT(" ");
+								Result += Children[i].IsValid() ? Children[i]->ToString(bPretty, IndentLevel + 2) : TEXT("nil");
+							}
+						}
+					}
+					else
+					{
+						Result += Newline + Indent + TEXT("  ");
+						Result += Children[i].IsValid() ? Children[i]->ToString(bPretty, IndentLevel + 1) : TEXT("nil");
+					}
+				}
+				Result += TEXT(")");
+			}
+			else
+			{
+				for (int32 i = 0; i < Children.Num(); i++)
+				{
+					if (i > 0) Result += TEXT(" ");
+					Result += Children[i].IsValid() ? Children[i]->ToString(false, 0) : TEXT("nil");
+				}
+				Result += TEXT(")");
+			}
+			
+			return Result;
+		}
+	}
+	
+	return TEXT("nil");
+}
+
+//------------------------------------------------------------------------------
+// FLispParseResult Implementation
+//------------------------------------------------------------------------------
+
+FLispParseResult FLispParseResult::Success(const TArray<FLispNodePtr>& InNodes)
+{
+	FLispParseResult Result;
+	Result.bSuccess = true;
+	Result.Nodes = InNodes;
+	return Result;
+}
+
+FLispParseResult FLispParseResult::Failure(const FString& InError, int32 InLine, int32 InColumn)
+{
+	FLispParseResult Result;
+	Result.bSuccess = false;
+	Result.Error = InError;
+	Result.ErrorLine = InLine;
+	Result.ErrorColumn = InColumn;
+	return Result;
+}
+
+//------------------------------------------------------------------------------
+// FLispParser Implementation
+//------------------------------------------------------------------------------
+
+FLispParser::FLispParser(const FString& InSource)
+	: Source(InSource)
+{
+}
+
+FLispParseResult FLispParser::Parse(const FString& Source)
+{
+	FLispParser Parser(Source);
+	return Parser.DoParse();
+}
+
+FLispParseResult FLispParser::DoParse()
+{
+	TArray<FLispNodePtr> Nodes;
+	
+	SkipWhitespaceAndComments();
+	
+	while (!IsAtEnd())
+	{
+		FLispNodePtr Expr = ParseExpr();
+		if (!Expr.IsValid())
+		{
+			return MakeError(TEXT("Failed to parse expression"));
+		}
+		Nodes.Add(Expr);
+		SkipWhitespaceAndComments();
+	}
+	
+	return FLispParseResult::Success(Nodes);
+}
+
+FLispNodePtr FLispParser::ParseExpr()
+{
+	SkipWhitespaceAndComments();
+	
+	if (IsAtEnd())
+	{
+		return nullptr;
+	}
+	
+	TCHAR c = Peek();
+	
+	if (c == '(')
+	{
+		return ParseList();
+	}
+	else if (c == '"')
+	{
+		return ParseString();
+	}
+	else if (c == ':')
+	{
+		return ParseSymbolOrKeyword();
+	}
+	else if (FChar::IsDigit(c) || (c == '-' && Position + 1 < Source.Len() && FChar::IsDigit(Source[Position + 1])))
+	{
+		return ParseNumber();
+	}
+	else
+	{
+		return ParseSymbolOrKeyword();
+	}
+}
+
+FLispNodePtr FLispParser::ParseList()
+{
+	int32 StartLine = Line;
+	int32 StartColumn = Column;
+	
+	if (!Match('('))
+	{
+		return nullptr;
+	}
+	
+	TArray<FLispNodePtr> Items;
+	SkipWhitespaceAndComments();
+	
+	while (!IsAtEnd() && Peek() != ')')
+	{
+		FLispNodePtr Item = ParseExpr();
+		if (!Item.IsValid())
+		{
+			return nullptr;
+		}
+		Items.Add(Item);
+		SkipWhitespaceAndComments();
+	}
+	
+	if (!Match(')'))
+	{
+		return nullptr;
+	}
+	
+	auto Node = FLispNode::MakeList(Items);
+	Node->Line = StartLine;
+	Node->Column = StartColumn;
+	return Node;
+}
+
+FLispNodePtr FLispParser::ParseString()
+{
+	int32 StartLine = Line;
+	int32 StartColumn = Column;
+	
+	if (!Match('"'))
+	{
+		return nullptr;
+	}
+	
+	FString Value;
+	
+	while (!IsAtEnd() && Peek() != '"')
+	{
+		TCHAR c = Advance();
+		
+		if (c == '\\' && !IsAtEnd())
+		{
+			TCHAR Escaped = Advance();
+			switch (Escaped)
+			{
+			case 'n': Value += '\n'; break;
+			case 't': Value += '\t'; break;
+			case 'r': Value += '\r'; break;
+			case '"': Value += '"'; break;
+			case '\\': Value += '\\'; break;
+			default: Value += Escaped; break;
+			}
+		}
+		else
+		{
+			Value += c;
+		}
+	}
+	
+	if (!Match('"'))
+	{
+		return nullptr;
+	}
+	
+	auto Node = FLispNode::MakeString(Value);
+	Node->Line = StartLine;
+	Node->Column = StartColumn;
+	return Node;
+}
+
+FLispNodePtr FLispParser::ParseNumber()
+{
+	int32 StartLine = Line;
+	int32 StartColumn = Column;
+	
+	FString NumStr;
+	
+	// Optional negative sign
+	if (Peek() == '-')
+	{
+		NumStr += Advance();
+	}
+	
+	// Integer part
+	while (!IsAtEnd() && FChar::IsDigit(Peek()))
+	{
+		NumStr += Advance();
+	}
+	
+	// Decimal part
+	if (!IsAtEnd() && Peek() == '.' && Position + 1 < Source.Len() && FChar::IsDigit(Source[Position + 1]))
+	{
+		NumStr += Advance(); // '.'
+		while (!IsAtEnd() && FChar::IsDigit(Peek()))
+		{
+			NumStr += Advance();
+		}
+	}
+	
+	// Exponent
+	if (!IsAtEnd() && (Peek() == 'e' || Peek() == 'E'))
+	{
+		NumStr += Advance();
+		if (!IsAtEnd() && (Peek() == '+' || Peek() == '-'))
+		{
+			NumStr += Advance();
+		}
+		while (!IsAtEnd() && FChar::IsDigit(Peek()))
+		{
+			NumStr += Advance();
+		}
+	}
+	
+	double Value = FCString::Atod(*NumStr);
+	
+	auto Node = FLispNode::MakeNumber(Value);
+	Node->StringValue = NumStr;  // Store original string to detect floats like "24.0"
+	Node->Line = StartLine;
+	Node->Column = StartColumn;
+	return Node;
+}
+
+FLispNodePtr FLispParser::ParseSymbolOrKeyword()
+{
+	int32 StartLine = Line;
+	int32 StartColumn = Column;
+	
+	bool bIsKeyword = (Peek() == ':');
+	FString Name;
+	
+	if (bIsKeyword)
+	{
+		Name += Advance(); // Include the ':'
+	}
+	
+	// Symbol chars: alphanumeric, underscore, hyphen, plus, minus, asterisk, slash, etc.
+	auto IsSymbolChar = [](TCHAR c) -> bool
+	{
+		return FChar::IsAlnum(c) || c == '_' || c == '-' || c == '+' || c == '*' || 
+			   c == '/' || c == '?' || c == '!' || c == '<' || c == '>' || c == '=' ||
+			   c == '.' || c == '@' || c == '&';
+	};
+	
+	while (!IsAtEnd() && IsSymbolChar(Peek()))
+	{
+		Name += Advance();
+	}
+	
+	if (Name.IsEmpty() || (bIsKeyword && Name.Len() == 1))
+	{
+		return nullptr;
+	}
+	
+	// Check for special symbols
+	if (Name.Equals(TEXT("nil"), ESearchCase::IgnoreCase))
+	{
+		return FLispNode::MakeNil();
+	}
+	if (Name.Equals(TEXT("true"), ESearchCase::IgnoreCase))
+	{
+		auto Node = FLispNode::MakeSymbol(TEXT("true"));
+		Node->Line = StartLine;
+		Node->Column = StartColumn;
+		return Node;
+	}
+	if (Name.Equals(TEXT("false"), ESearchCase::IgnoreCase))
+	{
+		auto Node = FLispNode::MakeSymbol(TEXT("false"));
+		Node->Line = StartLine;
+		Node->Column = StartColumn;
+		return Node;
+	}
+	
+	FLispNodePtr Node = bIsKeyword ? FLispNode::MakeKeyword(Name) : FLispNode::MakeSymbol(Name);
+	Node->Line = StartLine;
+	Node->Column = StartColumn;
+	return Node;
+}
+
+void FLispParser::SkipWhitespaceAndComments()
+{
+	while (!IsAtEnd())
+	{
+		TCHAR c = Peek();
+		
+		if (FChar::IsWhitespace(c))
+		{
+			Advance();
+		}
+		else if (c == ';')
+		{
+			// Line comment - skip to end of line
+			while (!IsAtEnd() && Peek() != '\n')
+			{
+				Advance();
+			}
+		}
+		else
+		{
+			break;
+		}
+	}
+}
+
+TCHAR FLispParser::Peek() const
+{
+	if (IsAtEnd()) return '\0';
+	return Source[Position];
+}
+
+TCHAR FLispParser::Advance()
+{
+	if (IsAtEnd()) return '\0';
+	
+	TCHAR c = Source[Position++];
+	
+	if (c == '\n')
+	{
+		Line++;
+		Column = 1;
+	}
+	else
+	{
+		Column++;
+	}
+	
+	return c;
+}
+
+bool FLispParser::IsAtEnd() const
+{
+	return Position >= Source.Len();
+}
+
+bool FLispParser::Match(TCHAR Expected)
+{
+	if (IsAtEnd() || Peek() != Expected)
+	{
+		return false;
+	}
+	Advance();
+	return true;
+}
+
+FLispParseResult FLispParser::MakeError(const FString& Message)
+{
+	return FLispParseResult::Failure(Message, Line, Column);
+}
+
+//------------------------------------------------------------------------------
+// FBlueprintLispResult Implementation
+//------------------------------------------------------------------------------
+
+FBlueprintLispResult FBlueprintLispResult::Success(const FString& Code)
+{
+	FBlueprintLispResult Result;
+	Result.bSuccess = true;
+	Result.LispCode = Code;
+	return Result;
+}
+
+FBlueprintLispResult FBlueprintLispResult::SuccessWithData(TSharedPtr<FJsonObject> Data)
+{
+	FBlueprintLispResult Result;
+	Result.bSuccess = true;
+	Result.GraphData = Data;
+	return Result;
+}
+
+FBlueprintLispResult FBlueprintLispResult::Failure(const FString& InError)
+{
+	FBlueprintLispResult Result;
+	Result.bSuccess = false;
+	Result.Error = InError;
+	return Result;
+}
+
+//------------------------------------------------------------------------------
+// FBlueprintLispConverter Implementation
+//------------------------------------------------------------------------------
+
+FBlueprintLispResult FBlueprintLispConverter::ValidateLisp(const FString& LispCode)
+{
+	FLispParseResult ParseResult = FLispParser::Parse(LispCode);
+	
+	if (!ParseResult.bSuccess)
+	{
+		return FBlueprintLispResult::Failure(FString::Printf(
+			TEXT("Parse error at line %d, column %d: %s"),
+			ParseResult.ErrorLine, ParseResult.ErrorColumn, *ParseResult.Error));
+	}
+	
+	// Validate each top-level form
+	for (const auto& Node : ParseResult.Nodes)
+	{
+		if (!Node.IsValid() || !Node->IsList())
+		{
+			return FBlueprintLispResult::Failure(TEXT("Top-level expressions must be lists"));
+		}
+		
+		FString FormName = Node->GetFormName();
+		if (FormName.IsEmpty())
+		{
+			return FBlueprintLispResult::Failure(TEXT("List must start with a symbol"));
+		}
+		
+		// Check for valid top-level forms
+		static const TSet<FString> ValidForms = {
+			TEXT("event"), TEXT("func"), TEXT("macro"), TEXT("var"), TEXT("comment")
+		};
+		
+		if (!ValidForms.Contains(FormName.ToLower()))
+		{
+			return FBlueprintLispResult::Failure(FString::Printf(
+				TEXT("Unknown top-level form: %s. Expected: event, func, macro, var"), *FormName));
+		}
+	}
+	
+	FBlueprintLispResult Result;
+	Result.bSuccess = true;
+	Result.LispCode = LispCode;
+	return Result;
+}
+
+FBlueprintLispResult FBlueprintLispConverter::BlueprintToLisp(
+	const FString& BlueprintPath,
+	const FString& GraphName,
+	const FBlueprintToLispOptions& Options)
+{
+	// This will be implemented to work with ECABridge's existing Blueprint introspection
+	// For now, return a placeholder
+	return FBlueprintLispResult::Failure(TEXT("BlueprintToLisp not yet implemented - use the ECA command"));
+}
+
+FBlueprintLispResult FBlueprintLispConverter::LispToBlueprint(
+	const FString& LispCode,
+	const FString& BlueprintPath,
+	const FString& GraphName,
+	const FLispToBlueprintOptions& Options)
+{
+	// Parse the Lisp code
+	FLispParseResult ParseResult = FLispParser::Parse(LispCode);
+	
+	if (!ParseResult.bSuccess)
+	{
+		return FBlueprintLispResult::Failure(FString::Printf(
+			TEXT("Parse error at line %d, column %d: %s"),
+			ParseResult.ErrorLine, ParseResult.ErrorColumn, *ParseResult.Error));
+	}
+	
+	// Create the graph data structure
+	TSharedPtr<FJsonObject> GraphData = MakeShared<FJsonObject>();
+	GraphData->SetStringField(TEXT("blueprint_path"), BlueprintPath);
+	GraphData->SetStringField(TEXT("graph_name"), GraphName);
+	
+	TArray<TSharedPtr<FJsonValue>> NodesArray;
+	TArray<TSharedPtr<FJsonValue>> ConnectionsArray;
+	
+	GraphData->SetArrayField(TEXT("nodes"), NodesArray);
+	GraphData->SetArrayField(TEXT("connections"), ConnectionsArray);
+	
+	// Process each top-level form
+	TMap<FString, FString> VariableMap; // Lisp variable name -> node ID
+	
+	for (const auto& Form : ParseResult.Nodes)
+	{
+		FString Error;
+		if (!ProcessForm(Form, GraphData, VariableMap, Error))
+		{
+			return FBlueprintLispResult::Failure(Error);
+		}
+	}
+	
+	return FBlueprintLispResult::SuccessWithData(GraphData);
+}
+
+bool FBlueprintLispConverter::ProcessForm(
+	const FLispNodePtr& Form,
+	TSharedPtr<FJsonObject> GraphData,
+	TMap<FString, FString>& VariableMap,
+	FString& OutError)
+{
+	if (!Form.IsValid() || !Form->IsList() || Form->Num() == 0)
+	{
+		OutError = TEXT("Invalid form");
+		return false;
+	}
+	
+	FString FormName = Form->GetFormName().ToLower();
+	
+	// TODO: Implement form processing
+	// This is where we'll convert each Lisp form to Blueprint nodes
+	
+	if (FormName == TEXT("event"))
+	{
+		// (event EventName body...)
+		if (Form->Num() < 2)
+		{
+			OutError = TEXT("event requires at least an event name");
+			return false;
+		}
+		
+		// Get event name
+		FLispNodePtr EventNameNode = Form->Get(1);
+		if (!EventNameNode.IsValid() || !EventNameNode->IsSymbol())
+		{
+			OutError = TEXT("event name must be a symbol");
+			return false;
+		}
+		
+		FString EventName = EventNameNode->StringValue;
+		
+		// TODO: Create event node and process body
+		// For now, just validate
+		return true;
+	}
+	else if (FormName == TEXT("func"))
+	{
+		// (func FuncName :inputs (...) :outputs (...) body...)
+		return true;
+	}
+	else if (FormName == TEXT("macro"))
+	{
+		// (macro MacroName :inputs (...) body...)
+		return true;
+	}
+	else if (FormName == TEXT("var"))
+	{
+		// (var VarName Type)
+		return true;
+	}
+	else if (FormName == TEXT("comment"))
+	{
+		// (comment "text")
+		return true;
+	}
+	
+	OutError = FString::Printf(TEXT("Unknown form: %s"), *FormName);
+	return false;
+}
+
+TSharedPtr<FJsonObject> FBlueprintLispConverter::CreateNodeJson(
+	const FString& NodeType,
+	const TMap<FString, FString>& Properties)
+{
+	TSharedPtr<FJsonObject> Node = MakeShared<FJsonObject>();
+	Node->SetStringField(TEXT("node_type"), NodeType);
+	
+	for (const auto& Pair : Properties)
+	{
+		Node->SetStringField(Pair.Key, Pair.Value);
+	}
+	
+	return Node;
+}
+
+bool FBlueprintLispConverter::ResolveExpression(
+	const FLispNodePtr& Expr,
+	TSharedPtr<FJsonObject> GraphData,
+	TMap<FString, FString>& VariableMap,
+	FString& OutNodeId,
+	FString& OutPinName,
+	FString& OutError)
+{
+	// TODO: Implement expression resolution
+	// This converts a Lisp expression to a node/pin reference
+	return true;
+}
+
+//------------------------------------------------------------------------------
+// Utility Functions Implementation
+//------------------------------------------------------------------------------
+
+namespace BlueprintLisp
+{
+
+FString PrettyPrint(const FString& LispCode)
+{
+	FLispParseResult ParseResult = FLispParser::Parse(LispCode);
+	if (!ParseResult.bSuccess)
+	{
+		return LispCode; // Return original if parse fails
+	}
+	
+	FString Result;
+	for (int32 i = 0; i < ParseResult.Nodes.Num(); i++)
+	{
+		if (i > 0)
+		{
+			Result += TEXT("\n\n");
+		}
+		Result += ParseResult.Nodes[i]->ToString(true, 0);
+	}
+	
+	return Result;
+}
+
+FString Minify(const FString& LispCode)
+{
+	FLispParseResult ParseResult = FLispParser::Parse(LispCode);
+	if (!ParseResult.bSuccess)
+	{
+		return LispCode;
+	}
+	
+	FString Result;
+	for (int32 i = 0; i < ParseResult.Nodes.Num(); i++)
+	{
+		if (i > 0)
+		{
+			Result += TEXT(" ");
+		}
+		Result += ParseResult.Nodes[i]->ToString(false, 0);
+	}
+	
+	return Result;
+}
+
+TArray<FString> ExtractSymbols(const FString& LispCode)
+{
+	TSet<FString> SymbolSet;
+	
+	FLispParseResult ParseResult = FLispParser::Parse(LispCode);
+	if (!ParseResult.bSuccess)
+	{
+		return TArray<FString>();
+	}
+	
+	TFunction<void(const FLispNodePtr&)> CollectSymbols = [&](const FLispNodePtr& Node)
+	{
+		if (!Node.IsValid()) return;
+		
+		if (Node->IsSymbol())
+		{
+			SymbolSet.Add(Node->StringValue);
+		}
+		else if (Node->IsList())
+		{
+			for (const auto& Child : Node->Children)
+			{
+				CollectSymbols(Child);
+			}
+		}
+	};
+	
+	for (const auto& Node : ParseResult.Nodes)
+	{
+		CollectSymbols(Node);
+	}
+	
+	return SymbolSet.Array();
+}
+
+bool IsValidSymbol(const FString& Name)
+{
+	if (Name.IsEmpty())
+	{
+		return false;
+	}
+	
+	// First char must be letter or underscore
+	TCHAR First = Name[0];
+	if (!FChar::IsAlpha(First) && First != '_')
+	{
+		return false;
+	}
+	
+	// Rest can be alphanumeric, underscore, or hyphen
+	for (int32 i = 1; i < Name.Len(); i++)
+	{
+		TCHAR c = Name[i];
+		if (!FChar::IsAlnum(c) && c != '_' && c != '-')
+		{
+			return false;
+		}
+	}
+	
+	return true;
+}
+
+} // namespace BlueprintLisp
diff --git a/Engine/Plugins/Experimental/ECABridge/Source/ECABridge/Private/Commands/BlueprintAutoLayout.cpp b/Engine/Plugins/Experimental/ECABridge/Source/ECABridge/Private/Commands/BlueprintAutoLayout.cpp
new file mode 100644
--- /dev/null
+++ b/Engine/Plugins/Experimental/ECABridge/Source/ECABridge/Private/Commands/BlueprintAutoLayout.cpp
@@ -0,0 +1,611 @@
+// Copyright Epic Games, Inc. All Rights Reserved.
+// BlueprintAutoLayout.cpp - Improved Blueprint graph layout algorithm
+
+#include "BlueprintAutoLayout.h"
+#include "K2Node_CallFunction.h"
+#include "K2Node_VariableGet.h"
+#include "K2Node_VariableSet.h"
+#include "K2Node_ExecutionSequence.h"
+
+int32 FBlueprintAutoLayout::LayoutGraph(UEdGraph* Graph, int32 StartX, int32 StartY)
+{
+	if (!Graph)
+	{
+		return 0;
+	}
+	
+	// Clear previous state
+	NodeInfoMap.Empty();
+	RootNodes.Empty();
+	AllExecNodes.Empty();
+	AllPureNodes.Empty();
+	PositionedPureNodes.Empty();
+	
+	// Phase 1: Build the layout tree from exec flow
+	BuildLayoutTree(Graph);
+	
+	if (RootNodes.Num() == 0)
+	{
+		return 0;
+	}
+	
+	// Phase 2: Calculate subtree heights (bottom-up, in pixels)
+	CalculateSubtreeHeights();
+	
+	// Phase 3: Assign positions (top-down)
+	AssignPositions(StartX, StartY);
+	
+	// Phase 4: Apply to actual nodes
+	ApplyPositions();
+	
+	return NodeInfoMap.Num();
+}
+
+int32 FBlueprintAutoLayout::LayoutSubtree(UEdGraph* Graph, const TArray<UEdGraphNode*>& SpecificRoots, int32 StartX, int32 StartY)
+{
+	if (!Graph || SpecificRoots.Num() == 0)
+	{
+		return 0;
+	}
+	
+	NodeInfoMap.Empty();
+	RootNodes.Empty();
+	AllExecNodes.Empty();
+	AllPureNodes.Empty();
+	PositionedPureNodes.Empty();
+	
+	BuildLayoutTree(Graph, &SpecificRoots);
+	
+	if (RootNodes.Num() == 0)
+	{
+		return 0;
+	}
+	
+	CalculateSubtreeHeights();
+	AssignPositions(StartX, StartY);
+	ApplyPositions();
+	
+	return NodeInfoMap.Num();
+}
+
+//------------------------------------------------------------------------------
+// Phase 1: Build Layout Tree
+//------------------------------------------------------------------------------
+
+void FBlueprintAutoLayout::BuildLayoutTree(UEdGraph* Graph, const TArray<UEdGraphNode*>* SpecificRoots)
+{
+	// First pass: Create info for all nodes and classify them
+	for (UEdGraphNode* Node : Graph->Nodes)
+	{
+		if (Node)
+		{
+			FLayoutNodeInfo* Info = GetOrCreateNodeInfo(Node);
+			ClassifyNode(Info);
+			
+			if (Info->bIsPureNode)
+			{
+				AllPureNodes.Add(Info);
+			}
+		}
+	}
+	
+	// Find root nodes
+	if (SpecificRoots)
+	{
+		for (UEdGraphNode* Root : *SpecificRoots)
+		{
+			if (FLayoutNodeInfo* Info = NodeInfoMap.Find(Root))
+			{
+				Info->bIsRootNode = true;
+				RootNodes.Add(Info);
+			}
+		}
+	}
+	else
+	{
+		// Auto-detect roots: events and function entries
+		for (auto& Pair : NodeInfoMap)
+		{
+			FLayoutNodeInfo& Info = Pair.Value;
+			if (Info.Node->IsA<UK2Node_Event>() || Info.Node->IsA<UK2Node_FunctionEntry>())
+			{
+				Info.bIsRootNode = true;
+				RootNodes.Add(&Info);
+			}
+		}
+		
+		// If no event/entry nodes, find nodes with exec output but no exec input
+		if (RootNodes.Num() == 0)
+		{
+			for (auto& Pair : NodeInfoMap)
+			{
+				FLayoutNodeInfo& Info = Pair.Value;
+				if (Info.bIsPureNode) continue;
+				
+				TArray<UEdGraphPin*> ExecInputs = GetExecInputPins(Info.Node);
+				TArray<UEdGraphPin*> ExecOutputs = GetExecOutputPins(Info.Node);
+				
+				bool bHasConnectedExecInput = false;
+				for (UEdGraphPin* Pin : ExecInputs)
+				{
+					if (Pin->LinkedTo.Num() > 0)
+					{
+						bHasConnectedExecInput = true;
+						break;
+					}
+				}
+				
+				if (!bHasConnectedExecInput && ExecOutputs.Num() > 0)
+				{
+					Info.bIsRootNode = true;
+					RootNodes.Add(&Info);
+				}
+			}
+		}
+	}
+	
+	// Sort roots by original Y position for consistent ordering
+	RootNodes.Sort([](const FLayoutNodeInfo& A, const FLayoutNodeInfo& B) {
+		return A.Node->NodePosY < B.Node->NodePosY;
+	});
+	
+	// Traverse from each root to build tree structure
+	for (FLayoutNodeInfo* Root : RootNodes)
+	{
+		TraverseExecFlow(Root, 0);
+	}
+	
+	// Collect pure node providers for each exec node
+	for (FLayoutNodeInfo* ExecNode : AllExecNodes)
+	{
+		CollectPureProviders(ExecNode);
+	}
+}
+
+FLayoutNodeInfo* FBlueprintAutoLayout::GetOrCreateNodeInfo(UEdGraphNode* Node)
+{
+	if (FLayoutNodeInfo* Existing = NodeInfoMap.Find(Node))
+	{
+		return Existing;
+	}
+	
+	FLayoutNodeInfo& NewInfo = NodeInfoMap.Add(Node);
+	NewInfo.Node = Node;
+	return &NewInfo;
+}
+
+void FBlueprintAutoLayout::ClassifyNode(FLayoutNodeInfo* Info)
+{
+	if (!Info || !Info->Node) return;
+	
+	Info->bIsPureNode = IsPureNode(Info->Node);
+	Info->bIsBranchNode = IsBranchNode(Info->Node);
+	
+	// Calculate actual node dimensions
+	CalculateNodeDimensions(Info);
+}
+
+void FBlueprintAutoLayout::CalculateNodeDimensions(FLayoutNodeInfo* Info)
+{
+	if (!Info || !Info->Node) return;
+	
+	UEdGraphNode* Node = Info->Node;
+	
+	// Try to get actual dimensions from node
+	int32 Width = Node->NodeWidth;
+	int32 Height = Node->NodeHeight;
+	
+	// If node doesn't have valid dimensions, estimate from pins
+	if (Width <= 0)
+	{
+		// Estimate width based on node title length and type
+		FString Title = Node->GetNodeTitle(ENodeTitleType::FullTitle).ToString();
+		Width = FMath::Max(Config.DefaultNodeWidth, Title.Len() * 8 + 60);
+		
+		// Call function nodes tend to be wider
+		if (Node->IsA<UK2Node_CallFunction>())
+		{
+			Width = FMath::Max(Width, 250);
+		}
+	}
+	
+	if (Height <= 0)
+	{
+		// Estimate height based on pin count
+		int32 InputPins = 0;
+		int32 OutputPins = 0;
+		
+		for (UEdGraphPin* Pin : Node->Pins)
+		{
+			if (Pin && !Pin->bHidden)
+			{
+				if (Pin->Direction == EGPD_Input)
+					InputPins++;
+				else
+					OutputPins++;
+			}
+		}
+		
+		int32 MaxPins = FMath::Max(InputPins, OutputPins);
+		Height = FMath::Max(Config.DefaultNodeHeight, 40 + MaxPins * Config.PinHeightEstimate);
+	}
+	
+	Info->NodeWidth = Width;
+	Info->NodeHeight = Height;
+}
+
+void FBlueprintAutoLayout::TraverseExecFlow(FLayoutNodeInfo* Current, int32 CurrentDepth)
+{
+	if (!Current || Current->bVisited || Current->bIsPureNode) return;
+	
+	Current->bVisited = true;
+	Current->Depth = CurrentDepth;
+	AllExecNodes.Add(Current);
+	
+	// Get exec output pins
+	TArray<UEdGraphPin*> ExecOutputs = GetExecOutputPins(Current->Node);
+	
+	for (UEdGraphPin* ExecOut : ExecOutputs)
+	{
+		for (UEdGraphPin* LinkedPin : ExecOut->LinkedTo)
+		{
+			if (LinkedPin)
+			{
+				UEdGraphNode* ChildNode = LinkedPin->GetOwningNode();
+				FLayoutNodeInfo* ChildInfo = GetOrCreateNodeInfo(ChildNode);
+				
+				if (!ChildInfo->bVisited && !ChildInfo->bIsPureNode)
+				{
+					ChildInfo->Parent = Current;
+					Current->ExecChildren.Add(ChildInfo);
+					Current->ChildPinNames.Add(ChildInfo, ExecOut->PinName);
+					
+					TraverseExecFlow(ChildInfo, CurrentDepth + 1);
+				}
+			}
+		}
+	}
+}
+
+void FBlueprintAutoLayout::CollectPureProviders(FLayoutNodeInfo* ExecNode)
+{
+	if (!ExecNode || !ExecNode->Node) return;
+	
+	TSet<UEdGraphNode*> Seen;
+	
+	for (UEdGraphPin* Pin : ExecNode->Node->Pins)
+	{
+		if (Pin && Pin->Direction == EGPD_Input && !IsExecPin(Pin))
+		{
+			for (UEdGraphPin* LinkedPin : Pin->LinkedTo)
+			{
+				if (LinkedPin)
+				{
+					UEdGraphNode* SourceNode = LinkedPin->GetOwningNode();
+					if (SourceNode && IsPureNode(SourceNode) && !Seen.Contains(SourceNode))
+					{
+						Seen.Add(SourceNode);
+						if (FLayoutNodeInfo* PureInfo = NodeInfoMap.Find(SourceNode))
+						{
+							ExecNode->DataProviders.Add(PureInfo);
+						}
+					}
+				}
+			}
+		}
+	}
+}
+
+//------------------------------------------------------------------------------
+// Phase 2: Calculate Subtree Heights (in pixels)
+//------------------------------------------------------------------------------
+
+void FBlueprintAutoLayout::CalculateSubtreeHeights()
+{
+	// Process nodes from deepest to shallowest (post-order)
+	AllExecNodes.Sort([](const FLayoutNodeInfo& A, const FLayoutNodeInfo& B) {
+		return A.Depth > B.Depth;
+	});
+	
+	for (FLayoutNodeInfo* Node : AllExecNodes)
+	{
+		Node->SubtreeHeight = CalculateHeight(Node);
+	}
+}
+
+int32 FBlueprintAutoLayout::CalculateHeight(FLayoutNodeInfo* Node)
+{
+	if (!Node) return Config.DefaultNodeHeight;
+	
+	// Use actual node height
+	int32 OwnHeight = Node->NodeHeight + Config.NodePaddingY;
+	
+	if (Node->ExecChildren.Num() == 0)
+	{
+		// Leaf node - just its own height
+		return OwnHeight;
+	}
+	
+	if (Node->bIsBranchNode && Node->ExecChildren.Num() >= 2)
+	{
+		// Branch node: children are stacked vertically (parallel paths)
+		// Total height = sum of all children heights + extra spacing between them
+		int32 TotalChildHeight = 0;
+		for (int32 i = 0; i < Node->ExecChildren.Num(); i++)
+		{
+			FLayoutNodeInfo* Child = Node->ExecChildren[i];
+			TotalChildHeight += Child->SubtreeHeight;
+			
+			// Add extra branch spacing between children (not after last)
+			if (i < Node->ExecChildren.Num() - 1)
+			{
+				TotalChildHeight += Config.BranchExtraPaddingY;
+			}
+		}
+		return FMath::Max(TotalChildHeight, OwnHeight);
+	}
+	else
+	{
+		// Sequential node: children are in sequence (same lane)
+		// Height = max of (own height, max child subtree height)
+		int32 MaxChildHeight = OwnHeight;
+		for (FLayoutNodeInfo* Child : Node->ExecChildren)
+		{
+			MaxChildHeight = FMath::Max(MaxChildHeight, Child->SubtreeHeight);
+		}
+		return MaxChildHeight;
+	}
+}
+
+//------------------------------------------------------------------------------
+// Phase 3: Assign Positions
+//------------------------------------------------------------------------------
+
+void FBlueprintAutoLayout::AssignPositions(int32 StartX, int32 StartY)
+{
+	int32 CurrentY = StartY;
+	
+	for (FLayoutNodeInfo* Root : RootNodes)
+	{
+		PositionExecSubtree(Root, StartX, CurrentY);
+		
+		// Move Y down by this root's subtree height plus extra spacing between roots
+		CurrentY += Root->SubtreeHeight + Config.RootExtraPaddingY;
+	}
+}
+
+void FBlueprintAutoLayout::PositionExecSubtree(FLayoutNodeInfo* Node, int32 X, int32 Y)
+{
+	if (!Node || Node->bPositioned) return;
+	
+	Node->bPositioned = true;
+	Node->LayoutX = X;
+	Node->LayoutY = Y;
+	
+	// Position pure nodes for this exec node
+	PositionPureNodesForConsumer(Node);
+	
+	// Position children - use actual node width + padding
+	int32 ChildX = X + Node->NodeWidth + Config.NodePaddingX;
+	
+	if (Node->bIsBranchNode && Node->ExecChildren.Num() >= 2)
+	{
+		// Branch: children go vertically stacked
+		// Sort children by pin name (Then before Else)
+		TArray<FLayoutNodeInfo*> SortedChildren = Node->ExecChildren;
+		SortedChildren.Sort([Node](const FLayoutNodeInfo& A, const FLayoutNodeInfo& B) {
+			FName PinA = Node->ChildPinNames.FindRef(const_cast<FLayoutNodeInfo*>(&A));
+			FName PinB = Node->ChildPinNames.FindRef(const_cast<FLayoutNodeInfo*>(&B));
+			if (PinA == TEXT("Then")) return true;
+			if (PinB == TEXT("Then")) return false;
+			return PinA.LexicalLess(PinB);
+		});
+		
+		int32 ChildY = Y; // First child at same Y as branch
+		
+		for (FLayoutNodeInfo* Child : SortedChildren)
+		{
+			PositionExecSubtree(Child, ChildX, ChildY);
+			
+			// Next child below this one's subtree (use actual subtree height + extra padding)
+			ChildY += Child->SubtreeHeight + Config.BranchExtraPaddingY;
+		}
+	}
+	else
+	{
+		// Sequential: children follow in same lane
+		for (FLayoutNodeInfo* Child : Node->ExecChildren)
+		{
+			PositionExecSubtree(Child, ChildX, Y);
+			// Subsequent sequential children continue using their actual width
+			ChildX += Child->NodeWidth + Config.NodePaddingX;
+		}
+	}
+}
+
+void FBlueprintAutoLayout::PositionPureNodesForConsumer(FLayoutNodeInfo* Consumer)
+{
+	if (!Consumer || Consumer->DataProviders.Num() == 0) return;
+	
+	// Position pure nodes in a column to the left of consumer
+	// Use the widest pure node's width for column spacing
+	int32 MaxPureWidth = Config.DefaultNodeWidth;
+	for (FLayoutNodeInfo* Pure : Consumer->DataProviders)
+	{
+		if (Pure)
+		{
+			MaxPureWidth = FMath::Max(MaxPureWidth, Pure->NodeWidth);
+		}
+	}
+	
+	// Start positioning to the left of consumer
+	int32 CurrentX = Consumer->LayoutX - MaxPureWidth - Config.NodePaddingX;
+	int32 CurrentY = Consumer->LayoutY;
+	
+	int32 Column = 0;
+	int32 RowInColumn = 0;
+	int32 ColumnStartY = CurrentY;
+	
+	for (FLayoutNodeInfo* Pure : Consumer->DataProviders)
+	{
+		if (!Pure) continue;
+		
+		// Skip if already positioned by another consumer
+		if (PositionedPureNodes.Contains(Pure))
+		{
+			continue;
+		}
+		
+		// Mark as positioned
+		PositionedPureNodes.Add(Pure);
+		Pure->bPositioned = true;
+		
+		// Calculate position - use actual heights for Y spacing
+		Pure->LayoutX = CurrentX - (Column * (MaxPureWidth + Config.NodePaddingX));
+		Pure->LayoutY = ColumnStartY + (RowInColumn * (Config.DefaultNodeHeight + Config.NodePaddingY));
+		
+		// Move to next slot
+		RowInColumn++;
+		if (RowInColumn >= Config.MaxPureNodesPerColumn)
+		{
+			RowInColumn = 0;
+			Column++;
+		}
+		
+		// Recursively position this pure node's pure inputs (they go further left)
+		PositionPureNodesForConsumer(Pure);
+	}
+}
+
+//------------------------------------------------------------------------------
+// Phase 4: Apply Positions
+//------------------------------------------------------------------------------
+
+void FBlueprintAutoLayout::ApplyPositions()
+{
+	for (auto& Pair : NodeInfoMap)
+	{
+		FLayoutNodeInfo& Info = Pair.Value;
+		if (Info.Node && Info.bPositioned)
+		{
+			Info.Node->NodePosX = Info.LayoutX;
+			Info.Node->NodePosY = Info.LayoutY;
+		}
+	}
+}
+
+//------------------------------------------------------------------------------
+// Helpers
+//------------------------------------------------------------------------------
+
+bool FBlueprintAutoLayout::IsExecPin(UEdGraphPin* Pin) const
+{
+	return Pin && Pin->PinType.PinCategory == UEdGraphSchema_K2::PC_Exec;
+}
+
+bool FBlueprintAutoLayout::IsPureNode(UEdGraphNode* Node) const
+{
+	if (!Node) return false;
+	
+	// Check if node has any exec pins
+	for (UEdGraphPin* Pin : Node->Pins)
+	{
+		if (IsExecPin(Pin))
+		{
+			return false;
+		}
+	}
+	
+	return true;
+}
+
+bool FBlueprintAutoLayout::IsBranchNode(UEdGraphNode* Node) const
+{
+	if (!Node) return false;
+	
+	// UK2Node_IfThenElse is definitely a branch
+	if (Node->IsA<UK2Node_IfThenElse>())
+	{
+		return true;
+	}
+	
+	// Sequence nodes have multiple outputs but are different - not branches
+	if (Node->IsA<UK2Node_ExecutionSequence>())
+	{
+		return false; // Treat sequence as sequential, not branching
+	}
+	
+	// Check for multiple exec outputs (switch, etc.)
+	TArray<UEdGraphPin*> ExecOutputs = GetExecOutputPins(Node);
+	return ExecOutputs.Num() > 1;
+}
+
+TArray<UEdGraphPin*> FBlueprintAutoLayout::GetExecOutputPins(UEdGraphNode* Node) const
+{
+	TArray<UEdGraphPin*> Result;
+	if (!Node) return Result;
+	
+	for (UEdGraphPin* Pin : Node->Pins)
+	{
+		if (Pin && Pin->Direction == EGPD_Output && IsExecPin(Pin))
+		{
+			Result.Add(Pin);
+		}
+	}
+	
+	// Sort for consistent ordering
+	Result.Sort([](const UEdGraphPin& A, const UEdGraphPin& B) {
+		if (A.PinName == TEXT("Then")) return true;
+		if (B.PinName == TEXT("Then")) return false;
+		if (A.PinName == TEXT("execute")) return true;
+		if (B.PinName == TEXT("execute")) return false;
+		return A.PinName.LexicalLess(B.PinName);
+	});
+	
+	return Result;
+}
+
+TArray<UEdGraphPin*> FBlueprintAutoLayout::GetExecInputPins(UEdGraphNode* Node) const
+{
+	TArray<UEdGraphPin*> Result;
+	if (!Node) return Result;
+	
+	for (UEdGraphPin* Pin : Node->Pins)
+	{
+		if (Pin && Pin->Direction == EGPD_Input && IsExecPin(Pin))
+		{
+			Result.Add(Pin);
+		}
+	}
+	
+	return Result;
+}
+
+TArray<UEdGraphNode*> FBlueprintAutoLayout::GetPureInputNodes(UEdGraphNode* Node) const
+{
+	TArray<UEdGraphNode*> Result;
+	TSet<UEdGraphNode*> Seen;
+	
+	if (!Node) return Result;
+	
+	for (UEdGraphPin* Pin : Node->Pins)
+	{
+		if (Pin && Pin->Direction == EGPD_Input && !IsExecPin(Pin))
+		{
+			for (UEdGraphPin* LinkedPin : Pin->LinkedTo)
+			{
+				if (LinkedPin)
+				{
+					UEdGraphNode* SourceNode = LinkedPin->GetOwningNode();
+					if (SourceNode && IsPureNode(SourceNode) && !Seen.Contains(SourceNode))
+					{
+						Seen.Add(SourceNode);
+						Result.Add(SourceNode);
+					}
+				}
+			}
+		}
+	}
+	
+	return Result;
+}
diff --git a/Engine/Plugins/Experimental/ECABridge/Source/ECABridge/Private/Commands/BlueprintAutoLayout.h b/Engine/Plugins/Experimental/ECABridge/Source/ECABridge/Private/Commands/BlueprintAutoLayout.h
new file mode 100644
--- /dev/null
+++ b/Engine/Plugins/Experimental/ECABridge/Source/ECABridge/Private/Commands/BlueprintAutoLayout.h
@@ -0,0 +1,145 @@
+// Copyright Epic Games, Inc. All Rights Reserved.
+// BlueprintAutoLayout.h - Improved Blueprint graph layout algorithm
+// Handles branches, loops, pure nodes, and complex execution flows
+
+#pragma once
+
+#include "CoreMinimal.h"
+#include "EdGraph/EdGraph.h"
+#include "EdGraph/EdGraphNode.h"
+#include "EdGraph/EdGraphPin.h"
+#include "K2Node_IfThenElse.h"
+#include "K2Node_Event.h"
+#include "K2Node_FunctionEntry.h"
+#include "K2Node_FunctionResult.h"
+#include "K2Node_MacroInstance.h"
+#include "EdGraphSchema_K2.h"
+
+/**
+ * Layout configuration options
+ */
+struct FBlueprintLayoutConfig
+{
+	// Padding/margins between nodes (added to actual node sizes)
+	int32 NodePaddingX = 50;           // Horizontal padding between nodes
+	int32 NodePaddingY = 30;           // Vertical padding between nodes
+	int32 BranchExtraPaddingY = 50;    // Extra vertical padding between branch paths
+	int32 RootExtraPaddingY = 80;      // Extra padding between different event roots
+	
+	// Pure node positioning
+	int32 PureNodeGapX = 30;           // Gap between pure nodes and their consumer
+	int32 PureNodePaddingY = 20;       // Vertical padding between pure nodes
+	int32 MaxPureNodesPerColumn = 4;   // Max pure nodes before starting new column
+	
+	// Fallback sizes when node dimensions aren't available
+	int32 DefaultNodeWidth = 220;      // Default width if node reports 0
+	int32 DefaultNodeHeight = 100;     // Default height if node reports 0
+	int32 PinHeightEstimate = 26;      // Estimated height per pin for size calculation
+};
+
+/**
+ * Internal node representation for layout calculations
+ */
+struct FLayoutNodeInfo
+{
+	UEdGraphNode* Node = nullptr;
+	
+	// Tree structure
+	FLayoutNodeInfo* Parent = nullptr;
+	TArray<FLayoutNodeInfo*> ExecChildren;     // Children via exec pins
+	TArray<FLayoutNodeInfo*> DataProviders;    // Pure nodes that feed into this node
+	
+	// Actual node dimensions (calculated from node or estimated)
+	int32 NodeWidth = 0;
+	int32 NodeHeight = 0;
+	
+	// Layout data
+	int32 Depth = 0;              // Distance from root in exec flow
+	
+	// Subtree measurements (in pixels)
+	int32 SubtreeHeight = 0;      // Total height needed for this subtree
+	int32 SubtreeWidth = 0;       // Total width needed for this subtree
+	
+	// Position (final calculated position)
+	int32 LayoutX = 0;
+	int32 LayoutY = 0;
+	
+	// Flags
+	bool bIsBranchNode = false;
+	bool bIsPureNode = false;
+	bool bIsRootNode = false;
+	bool bVisited = false;
+	bool bPositioned = false;     // Has this node been positioned already?
+	
+	// For branch nodes: track which pin each child came from
+	TMap<FLayoutNodeInfo*, FName> ChildPinNames;
+};
+
+/**
+ * Main layout algorithm class
+ */
+class FBlueprintAutoLayout
+{
+public:
+	FBlueprintAutoLayout(const FBlueprintLayoutConfig& InConfig = FBlueprintLayoutConfig())
+		: Config(InConfig)
+	{}
+	
+	/**
+	 * Layout all nodes in a graph
+	 */
+	int32 LayoutGraph(UEdGraph* Graph, int32 StartX = 0, int32 StartY = 0);
+	
+	/**
+	 * Layout a subtree starting from specific nodes
+	 */
+	int32 LayoutSubtree(UEdGraph* Graph, const TArray<UEdGraphNode*>& RootNodes, int32 StartX = 0, int32 StartY = 0);
+	
+	/** Get layout info for debugging */
+	const TMap<UEdGraphNode*, FLayoutNodeInfo>& GetLayoutInfo() const { return NodeInfoMap; }
+	
+private:
+	FBlueprintLayoutConfig Config;
+	TMap<UEdGraphNode*, FLayoutNodeInfo> NodeInfoMap;
+	TArray<FLayoutNodeInfo*> RootNodes;
+	TArray<FLayoutNodeInfo*> AllExecNodes;      // Only exec-flow nodes
+	TArray<FLayoutNodeInfo*> AllPureNodes;      // Only pure nodes
+	TSet<FLayoutNodeInfo*> PositionedPureNodes; // Track which pure nodes are already positioned
+	
+	// Phase 1: Build the layout tree
+	void BuildLayoutTree(UEdGraph* Graph, const TArray<UEdGraphNode*>* SpecificRoots = nullptr);
+	FLayoutNodeInfo* GetOrCreateNodeInfo(UEdGraphNode* Node);
+	void ClassifyNode(FLayoutNodeInfo* Info);
+	void CalculateNodeDimensions(FLayoutNodeInfo* Info);
+	void TraverseExecFlow(FLayoutNodeInfo* Current, int32 CurrentDepth);
+	void CollectPureProviders(FLayoutNodeInfo* ExecNode);
+	
+	// Phase 2: Calculate subtree dimensions (in pixels)
+	void CalculateSubtreeHeights();
+	int32 CalculateHeight(FLayoutNodeInfo* Node);
+	
+	// Phase 3: Assign positions
+	void AssignPositions(int32 StartX, int32 StartY);
+	void PositionExecSubtree(FLayoutNodeInfo* Node, int32 X, int32 Y);
+	void PositionPureNodesForConsumer(FLayoutNodeInfo* Consumer);
+	
+	// Phase 4: Apply positions to actual nodes
+	void ApplyPositions();
+	
+	// Helpers
+	bool IsExecPin(UEdGraphPin* Pin) const;
+	bool IsPureNode(UEdGraphNode* Node) const;
+	bool IsBranchNode(UEdGraphNode* Node) const;
+	TArray<UEdGraphPin*> GetExecOutputPins(UEdGraphNode* Node) const;
+	TArray<UEdGraphPin*> GetExecInputPins(UEdGraphNode* Node) const;
+	TArray<UEdGraphNode*> GetPureInputNodes(UEdGraphNode* Node) const;
+};
+
+/**
+ * Utility function to layout a graph with default settings
+ */
+inline int32 AutoLayoutBlueprintGraph(UEdGraph* Graph, int32 StartX = 0, int32 StartY = 0)
+{
+	FBlueprintAutoLayout Layout;
+	return Layout.LayoutGraph(Graph, StartX, StartY);
+}
diff --git a/Engine/Plugins/Experimental/ECABridge/Source/ECABridge/Private/Commands/ECAActorCommands.cpp b/Engine/Plugins/Experimental/ECABridge/Source/ECABridge/Private/Commands/ECAActorCommands.cpp
new file mode 100644
--- /dev/null
+++ b/Engine/Plugins/Experimental/ECABridge/Source/ECABridge/Private/Commands/ECAActorCommands.cpp
@@ -0,0 +1,1046 @@
+// Copyright Epic Games, Inc. All Rights Reserved.
+
+#include "Commands/ECAActorCommands.h"
+#include "Commands/ECACommand.h"
+#include "Editor.h"
+#include "Engine/World.h"
+#include "GameFramework/Actor.h"
+#include "Engine/StaticMeshActor.h"
+#include "Engine/PointLight.h"
+#include "Engine/SpotLight.h"
+#include "Engine/DirectionalLight.h"
+#include "Camera/CameraActor.h"
+#include "Components/StaticMeshComponent.h"
+#include "Engine/StaticMesh.h"
+#include "Materials/MaterialInterface.h"
+#include "EngineUtils.h"
+#include "EditorActorFolders.h"
+#include "Subsystems/EditorActorSubsystem.h"
+#include "Components/PrimitiveComponent.h"
+
+//------------------------------------------------------------------------------
+// Spatial Verification Helpers
+//------------------------------------------------------------------------------
+
+namespace SpatialHelpers
+{
+	// Convert FBox to JSON
+	TSharedPtr<FJsonObject> BoxToJson(const FBox& Box)
+	{
+		TSharedPtr<FJsonObject> BoxJson = MakeShared<FJsonObject>();
+		
+		TSharedPtr<FJsonObject> MinJson = MakeShared<FJsonObject>();
+		MinJson->SetNumberField(TEXT("x"), Box.Min.X);
+		MinJson->SetNumberField(TEXT("y"), Box.Min.Y);
+		MinJson->SetNumberField(TEXT("z"), Box.Min.Z);
+		BoxJson->SetObjectField(TEXT("min"), MinJson);
+		
+		TSharedPtr<FJsonObject> MaxJson = MakeShared<FJsonObject>();
+		MaxJson->SetNumberField(TEXT("x"), Box.Max.X);
+		MaxJson->SetNumberField(TEXT("y"), Box.Max.Y);
+		MaxJson->SetNumberField(TEXT("z"), Box.Max.Z);
+		BoxJson->SetObjectField(TEXT("max"), MaxJson);
+		
+		// Add size for convenience
+		FVector Size = Box.GetSize();
+		TSharedPtr<FJsonObject> SizeJson = MakeShared<FJsonObject>();
+		SizeJson->SetNumberField(TEXT("x"), Size.X);
+		SizeJson->SetNumberField(TEXT("y"), Size.Y);
+		SizeJson->SetNumberField(TEXT("z"), Size.Z);
+		BoxJson->SetObjectField(TEXT("size"), SizeJson);
+		
+		// Add center
+		FVector Center = Box.GetCenter();
+		TSharedPtr<FJsonObject> CenterJson = MakeShared<FJsonObject>();
+		CenterJson->SetNumberField(TEXT("x"), Center.X);
+		CenterJson->SetNumberField(TEXT("y"), Center.Y);
+		CenterJson->SetNumberField(TEXT("z"), Center.Z);
+		BoxJson->SetObjectField(TEXT("center"), CenterJson);
+		
+		return BoxJson;
+	}
+	
+	// Get spatial verification data for an actor
+	TSharedPtr<FJsonObject> GetSpatialData(AActor* Actor, UWorld* World)
+	{
+		if (!Actor || !World)
+		{
+			return nullptr;
+		}
+		
+		TSharedPtr<FJsonObject> SpatialJson = MakeShared<FJsonObject>();
+		
+		// World bounds
+		FBox WorldBounds = Actor->GetComponentsBoundingBox(true);
+		if (WorldBounds.IsValid)
+		{
+			SpatialJson->SetObjectField(TEXT("world_bounds"), BoxToJson(WorldBounds));
+		}
+		
+		// Distance to ground (trace down from actor origin)
+		FVector ActorLocation = Actor->GetActorLocation();
+		FHitResult HitResult;
+		FVector TraceStart = ActorLocation;
+		FVector TraceEnd = ActorLocation - FVector(0, 0, 100000); // Trace 1km down
+		
+		FCollisionQueryParams QueryParams;
+		QueryParams.AddIgnoredActor(Actor);
+		
+		if (World->LineTraceSingleByChannel(HitResult, TraceStart, TraceEnd, ECC_WorldStatic, QueryParams))
+		{
+			double DistanceToGround = FVector::Dist(ActorLocation, HitResult.Location);
+			SpatialJson->SetNumberField(TEXT("distance_to_ground"), DistanceToGround);
+			
+			// Ground hit info
+			TSharedPtr<FJsonObject> GroundJson = MakeShared<FJsonObject>();
+			GroundJson->SetNumberField(TEXT("x"), HitResult.Location.X);
+			GroundJson->SetNumberField(TEXT("y"), HitResult.Location.Y);
+			GroundJson->SetNumberField(TEXT("z"), HitResult.Location.Z);
+			if (HitResult.GetActor())
+			{
+				GroundJson->SetStringField(TEXT("hit_actor"), HitResult.GetActor()->GetActorLabel());
+			}
+			SpatialJson->SetObjectField(TEXT("ground_point"), GroundJson);
+		}
+		else
+		{
+			SpatialJson->SetNumberField(TEXT("distance_to_ground"), -1); // No ground found
+		}
+		
+		// Check for overlapping actors
+		TArray<TSharedPtr<FJsonValue>> OverlapsArray;
+		TArray<AActor*> OverlappingActors;
+		Actor->GetOverlappingActors(OverlappingActors);
+		
+		for (AActor* OtherActor : OverlappingActors)
+		{
+			if (OtherActor && OtherActor != Actor)
+			{
+				TSharedPtr<FJsonObject> OverlapJson = MakeShared<FJsonObject>();
+				OverlapJson->SetStringField(TEXT("actor"), OtherActor->GetActorLabel());
+				OverlapJson->SetStringField(TEXT("class"), OtherActor->GetClass()->GetName());
+				OverlapsArray.Add(MakeShared<FJsonValueObject>(OverlapJson));
+			}
+		}
+		SpatialJson->SetArrayField(TEXT("overlapping_actors"), OverlapsArray);
+		
+		// Find nearby actors within a radius
+		TArray<TSharedPtr<FJsonValue>> NearbyArray;
+		const float NearbyRadius = 500.0f; // 5 meters
+		
+		for (TActorIterator<AActor> It(World); It; ++It)
+		{
+			AActor* OtherActor = *It;
+			if (OtherActor && OtherActor != Actor && !OtherActor->IsHiddenEd())
+			{
+				float Distance = FVector::Dist(ActorLocation, OtherActor->GetActorLocation());
+				if (Distance <= NearbyRadius)
+				{
+					TSharedPtr<FJsonObject> NearbyJson = MakeShared<FJsonObject>();
+					NearbyJson->SetStringField(TEXT("actor"), OtherActor->GetActorLabel());
+					NearbyJson->SetNumberField(TEXT("distance"), Distance);
+					NearbyArray.Add(MakeShared<FJsonValueObject>(NearbyJson));
+				}
+			}
+		}
+		SpatialJson->SetArrayField(TEXT("nearby_actors"), NearbyArray);
+		
+		return SpatialJson;
+	}
+}
+
+// Register all actor commands
+REGISTER_ECA_COMMAND(FECACommand_GetActorsInLevel)
+REGISTER_ECA_COMMAND(FECACommand_CreateActor)
+REGISTER_ECA_COMMAND(FECACommand_DeleteActor)
+REGISTER_ECA_COMMAND(FECACommand_SetActorTransform)
+REGISTER_ECA_COMMAND(FECACommand_GetActorProperties)
+REGISTER_ECA_COMMAND(FECACommand_SetActorProperty)
+REGISTER_ECA_COMMAND(FECACommand_FindActors)
+REGISTER_ECA_COMMAND(FECACommand_DuplicateActor)
+
+
+//------------------------------------------------------------------------------
+// GetActorsInLevel
+//------------------------------------------------------------------------------
+
+FECACommandResult FECACommand_GetActorsInLevel::Execute(const TSharedPtr<FJsonObject>& Params)
+{
+	UWorld* World = GetEditorWorld();
+	if (!World)
+	{
+		return FECACommandResult::Error(TEXT("No world available"));
+	}
+	
+	FString ClassFilter;
+	bool bIncludeHidden = false;
+	GetStringParam(Params, TEXT("class_filter"), ClassFilter, false);
+	GetBoolParam(Params, TEXT("include_hidden"), bIncludeHidden, false);
+	
+	TArray<TSharedPtr<FJsonValue>> ActorsArray;
+	
+	for (TActorIterator<AActor> It(World); It; ++It)
+	{
+		AActor* Actor = *It;
+		
+		// Skip hidden actors unless requested
+		if (!bIncludeHidden && Actor->IsHiddenEd())
+		{
+			continue;
+		}
+		
+		// Apply class filter
+		if (!ClassFilter.IsEmpty())
+		{
+			FString ActorClassName = Actor->GetClass()->GetName();
+			if (!ActorClassName.Contains(ClassFilter))
+			{
+				continue;
+			}
+		}
+		
+		TSharedPtr<FJsonObject> ActorJson = MakeShared<FJsonObject>();
+		ActorJson->SetStringField(TEXT("name"), Actor->GetActorLabel());
+		ActorJson->SetStringField(TEXT("internal_name"), Actor->GetName());
+		ActorJson->SetStringField(TEXT("class"), Actor->GetClass()->GetName());
+		ActorJson->SetObjectField(TEXT("location"), VectorToJson(Actor->GetActorLocation()));
+		ActorJson->SetObjectField(TEXT("rotation"), RotatorToJson(Actor->GetActorRotation()));
+		ActorJson->SetObjectField(TEXT("scale"), VectorToJson(Actor->GetActorScale3D()));
+		ActorJson->SetBoolField(TEXT("is_hidden"), Actor->IsHiddenEd());
+		
+		// Include folder path
+		FName FolderPath = Actor->GetFolderPath();
+		if (!FolderPath.IsNone())
+		{
+			ActorJson->SetStringField(TEXT("folder"), FolderPath.ToString());
+		}
+		
+		// Include tags
+		if (Actor->Tags.Num() > 0)
+		{
+			TArray<TSharedPtr<FJsonValue>> TagsArray;
+			for (const FName& Tag : Actor->Tags)
+			{
+				TagsArray.Add(MakeShared<FJsonValueString>(Tag.ToString()));
+			}
+			ActorJson->SetArrayField(TEXT("tags"), TagsArray);
+		}
+		
+		ActorsArray.Add(MakeShared<FJsonValueObject>(ActorJson));
+	}
+	
+	TSharedPtr<FJsonObject> Result = MakeResult();
+	Result->SetArrayField(TEXT("actors"), ActorsArray);
+	Result->SetNumberField(TEXT("count"), ActorsArray.Num());
+	return FECACommandResult::Success(Result);
+}
+
+//------------------------------------------------------------------------------
+// CreateActor
+//------------------------------------------------------------------------------
+
+FECACommandResult FECACommand_CreateActor::Execute(const TSharedPtr<FJsonObject>& Params)
+{
+	FString ActorType;
+	if (!GetStringParam(Params, TEXT("actor_type"), ActorType))
+	{
+		return FECACommandResult::Error(TEXT("Missing required parameter: actor_type"));
+	}
+	
+	UWorld* World = GetEditorWorld();
+	if (!World)
+	{
+		return FECACommandResult::Error(TEXT("No world available"));
+	}
+	
+	// Parse transform parameters
+	FVector Location = FVector::ZeroVector;
+	FRotator Rotation = FRotator::ZeroRotator;
+	FVector Scale = FVector::OneVector;
+	
+	GetVectorParam(Params, TEXT("location"), Location, false);
+	GetRotatorParam(Params, TEXT("rotation"), Rotation, false);
+	GetVectorParam(Params, TEXT("scale"), Scale, false);
+	
+	FActorSpawnParameters SpawnParams;
+	SpawnParams.SpawnCollisionHandlingOverride = ESpawnActorCollisionHandlingMethod::AdjustIfPossibleButAlwaysSpawn;
+	
+	AActor* NewActor = nullptr;
+	
+	// Handle built-in actor types
+	if (ActorType.Equals(TEXT("StaticMeshActor"), ESearchCase::IgnoreCase))
+	{
+		AStaticMeshActor* MeshActor = World->SpawnActor<AStaticMeshActor>(Location, Rotation, SpawnParams);
+		if (MeshActor)
+		{
+			FString MeshPath;
+			if (GetStringParam(Params, TEXT("mesh"), MeshPath, false))
+			{
+				UStaticMesh* Mesh = LoadObject<UStaticMesh>(nullptr, *MeshPath);
+				if (Mesh && MeshActor->GetStaticMeshComponent())
+				{
+					MeshActor->GetStaticMeshComponent()->SetStaticMesh(Mesh);
+				}
+			}
+			else
+			{
+				// Default cube
+				UStaticMesh* CubeMesh = LoadObject<UStaticMesh>(nullptr, TEXT("/Engine/BasicShapes/Cube.Cube"));
+				if (CubeMesh && MeshActor->GetStaticMeshComponent())
+				{
+					MeshActor->GetStaticMeshComponent()->SetStaticMesh(CubeMesh);
+				}
+			}
+		}
+		NewActor = MeshActor;
+	}
+	else if (ActorType.Equals(TEXT("PointLight"), ESearchCase::IgnoreCase))
+	{
+		NewActor = World->SpawnActor<APointLight>(Location, Rotation, SpawnParams);
+	}
+	else if (ActorType.Equals(TEXT("SpotLight"), ESearchCase::IgnoreCase))
+	{
+		NewActor = World->SpawnActor<ASpotLight>(Location, Rotation, SpawnParams);
+	}
+	else if (ActorType.Equals(TEXT("DirectionalLight"), ESearchCase::IgnoreCase))
+	{
+		NewActor = World->SpawnActor<ADirectionalLight>(Location, Rotation, SpawnParams);
+	}
+	else if (ActorType.Equals(TEXT("CameraActor"), ESearchCase::IgnoreCase))
+	{
+		NewActor = World->SpawnActor<ACameraActor>(Location, Rotation, SpawnParams);
+	}
+	else
+	{
+		// Try to find class by path
+		UClass* ActorClass = LoadObject<UClass>(nullptr, *ActorType);
+		if (!ActorClass)
+		{
+			// Try as a short name
+			ActorClass = FindObject<UClass>(nullptr, *FString::Printf(TEXT("/Script/Engine.%s"), *ActorType));
+		}
+		
+		if (ActorClass && ActorClass->IsChildOf(AActor::StaticClass()))
+		{
+			NewActor = World->SpawnActor<AActor>(ActorClass, Location, Rotation, SpawnParams);
+		}
+		else
+		{
+			return FECACommandResult::Error(FString::Printf(TEXT("Unknown actor type: %s"), *ActorType));
+		}
+	}
+	
+	if (!NewActor)
+	{
+		return FECACommandResult::Error(TEXT("Failed to create actor"));
+	}
+	
+	// Apply scale
+	NewActor->SetActorScale3D(Scale);
+	
+	// Set name if specified
+	FString ActorName;
+	if (GetStringParam(Params, TEXT("name"), ActorName, false))
+	{
+		NewActor->SetActorLabel(*ActorName);
+	}
+	
+	// Put in folder if specified
+	FString FolderPath;
+	if (GetStringParam(Params, TEXT("folder"), FolderPath, false))
+	{
+		NewActor->SetFolderPath(*FolderPath);
+	}
+	
+	TSharedPtr<FJsonObject> Result = MakeResult();
+	Result->SetStringField(TEXT("actor_name"), NewActor->GetActorLabel());
+	Result->SetStringField(TEXT("internal_name"), NewActor->GetName());
+	Result->SetStringField(TEXT("class"), NewActor->GetClass()->GetName());
+	Result->SetObjectField(TEXT("transform"), TransformToJson(NewActor->GetTransform()));
+	
+	// Add spatial verification data
+	TSharedPtr<FJsonObject> SpatialData = SpatialHelpers::GetSpatialData(NewActor, World);
+	if (SpatialData.IsValid())
+	{
+		Result->SetObjectField(TEXT("spatial"), SpatialData);
+	}
+	
+	return FECACommandResult::Success(Result);
+}
+
+//------------------------------------------------------------------------------
+// DeleteActor
+//------------------------------------------------------------------------------
+
+FECACommandResult FECACommand_DeleteActor::Execute(const TSharedPtr<FJsonObject>& Params)
+{
+	FString ActorName;
+	if (!GetStringParam(Params, TEXT("actor_name"), ActorName))
+	{
+		return FECACommandResult::Error(TEXT("Missing required parameter: actor_name"));
+	}
+	
+	AActor* Actor = FindActorByName(ActorName);
+	if (!Actor)
+	{
+		return FECACommandResult::Error(FString::Printf(TEXT("Actor not found: %s"), *ActorName));
+	}
+	
+	FString DeletedName = Actor->GetActorLabel();
+	Actor->Destroy();
+	
+	TSharedPtr<FJsonObject> Result = MakeResult();
+	Result->SetStringField(TEXT("deleted_actor"), DeletedName);
+	return FECACommandResult::Success(Result);
+}
+
+//------------------------------------------------------------------------------
+// SetActorTransform
+//------------------------------------------------------------------------------
+
+FECACommandResult FECACommand_SetActorTransform::Execute(const TSharedPtr<FJsonObject>& Params)
+{
+	FString ActorName;
+	if (!GetStringParam(Params, TEXT("actor_name"), ActorName))
+	{
+		return FECACommandResult::Error(TEXT("Missing required parameter: actor_name"));
+	}
+	
+	AActor* Actor = FindActorByName(ActorName);
+	if (!Actor)
+	{
+		return FECACommandResult::Error(FString::Printf(TEXT("Actor not found: %s"), *ActorName));
+	}
+	
+	// Check if actor has a root component
+	USceneComponent* RootComp = Actor->GetRootComponent();
+	if (!RootComp)
+	{
+		return FECACommandResult::Error(FString::Printf(TEXT("Actor '%s' has no root component"), *ActorName));
+	}
+	
+	// Get current transform as base
+	FVector NewLocation = Actor->GetActorLocation();
+	FRotator NewRotation = Actor->GetActorRotation();
+	FVector NewScale = Actor->GetActorScale3D();
+	
+	bool bLocationProvided = false;
+	bool bRotationProvided = false;
+	bool bScaleProvided = false;
+	
+	// Get new values if provided
+	FVector Location;
+	if (GetVectorParam(Params, TEXT("location"), Location, false))
+	{
+		NewLocation = Location;
+		bLocationProvided = true;
+	}
+	
+	FRotator Rotation;
+	if (GetRotatorParam(Params, TEXT("rotation"), Rotation, false))
+	{
+		NewRotation = Rotation;
+		bRotationProvided = true;
+	}
+	
+	FVector Scale;
+	if (GetVectorParam(Params, TEXT("scale"), Scale, false))
+	{
+		NewScale = Scale;
+		bScaleProvided = true;
+	}
+	
+	// For editor operations, we need to use Modify() to properly record undo
+	// and notify the editor of changes - this is how the viewport transform tools work
+	RootComp->Modify();
+	Actor->Modify();
+	
+	// Set transform values directly on the root component
+	// Using SetRelativeLocation_Direct etc. bypasses mobility checks
+	// This is the same approach used by the editor's transform gizmo
+	if (bLocationProvided)
+	{
+		RootComp->SetWorldLocation(NewLocation, false, nullptr, ETeleportType::ResetPhysics);
+	}
+	
+	if (bRotationProvided)
+	{
+		RootComp->SetWorldRotation(NewRotation, false, nullptr, ETeleportType::ResetPhysics);
+	}
+	
+	if (bScaleProvided)
+	{
+		RootComp->SetWorldScale3D(NewScale);
+	}
+	
+	// Force update transform - crucial for static meshes
+	RootComp->UpdateComponentToWorld();
+	
+	// Notify that the actor has been moved (triggers editor updates)
+	GEditor->BroadcastLevelActorListChanged();
+	Actor->PostEditMove(true);
+	
+	// Mark package dirty so changes can be saved
+	Actor->MarkPackageDirty();
+	
+	// Refresh all viewports to show the change
+	GEditor->RedrawAllViewports();
+	
+	// Verify the transform was actually applied
+	FVector FinalLocation = Actor->GetActorLocation();
+	FRotator FinalRotation = Actor->GetActorRotation();
+	FVector FinalScale = Actor->GetActorScale3D();
+	
+	TSharedPtr<FJsonObject> Result = MakeResult();
+	Result->SetObjectField(TEXT("transform"), TransformToJson(Actor->GetTransform()));
+	Result->SetBoolField(TEXT("success"), true);
+	
+	// Report what was changed
+	TArray<FString> ChangedFields;
+	if (bLocationProvided) ChangedFields.Add(TEXT("location"));
+	if (bRotationProvided) ChangedFields.Add(TEXT("rotation"));
+	if (bScaleProvided) ChangedFields.Add(TEXT("scale"));
+	Result->SetStringField(TEXT("changed"), FString::Join(ChangedFields, TEXT(", ")));
+	
+	// Verify transform was applied (compare requested vs actual)
+	bool bLocationMatch = !bLocationProvided || FinalLocation.Equals(NewLocation, 0.01f);
+	bool bRotationMatch = !bRotationProvided || FinalRotation.Equals(NewRotation, 0.01f);
+	bool bScaleMatch = !bScaleProvided || FinalScale.Equals(NewScale, 0.01f);
+	
+	if (!bLocationMatch || !bRotationMatch || !bScaleMatch)
+	{
+		Result->SetStringField(TEXT("warning"), TEXT("Transform may not have been fully applied - check mobility or constraints"));
+	}
+	
+	// Add spatial verification data
+	UWorld* World = GetEditorWorld();
+	TSharedPtr<FJsonObject> SpatialData = SpatialHelpers::GetSpatialData(Actor, World);
+	if (SpatialData.IsValid())
+	{
+		Result->SetObjectField(TEXT("spatial"), SpatialData);
+	}
+	
+	return FECACommandResult::Success(Result);
+}
+
+//------------------------------------------------------------------------------
+// GetActorProperties
+//------------------------------------------------------------------------------
+
+FECACommandResult FECACommand_GetActorProperties::Execute(const TSharedPtr<FJsonObject>& Params)
+{
+	FString ActorName;
+	if (!GetStringParam(Params, TEXT("actor_name"), ActorName))
+	{
+		return FECACommandResult::Error(TEXT("Missing required parameter: actor_name"));
+	}
+	
+	AActor* Actor = FindActorByName(ActorName);
+	if (!Actor)
+	{
+		return FECACommandResult::Error(FString::Printf(TEXT("Actor not found: %s"), *ActorName));
+	}
+	
+	TSharedPtr<FJsonObject> Result = MakeResult();
+	Result->SetStringField(TEXT("name"), Actor->GetActorLabel());
+	Result->SetStringField(TEXT("internal_name"), Actor->GetName());
+	Result->SetStringField(TEXT("class"), Actor->GetClass()->GetName());
+	Result->SetStringField(TEXT("class_path"), Actor->GetClass()->GetPathName());
+	Result->SetObjectField(TEXT("transform"), TransformToJson(Actor->GetTransform()));
+	Result->SetBoolField(TEXT("is_hidden"), Actor->IsHiddenEd());
+	
+	// Mobility
+	USceneComponent* RootComponent = Actor->GetRootComponent();
+	if (RootComponent)
+	{
+		FString Mobility;
+		switch (RootComponent->Mobility)
+		{
+			case EComponentMobility::Static: Mobility = TEXT("Static"); break;
+			case EComponentMobility::Stationary: Mobility = TEXT("Stationary"); break;
+			case EComponentMobility::Movable: Mobility = TEXT("Movable"); break;
+		}
+		Result->SetStringField(TEXT("mobility"), Mobility);
+	}
+	
+	// Folder
+	FName FolderPath = Actor->GetFolderPath();
+	if (!FolderPath.IsNone())
+	{
+		Result->SetStringField(TEXT("folder"), FolderPath.ToString());
+	}
+	
+	// Tags
+	TArray<TSharedPtr<FJsonValue>> TagsArray;
+	for (const FName& Tag : Actor->Tags)
+	{
+		TagsArray.Add(MakeShared<FJsonValueString>(Tag.ToString()));
+	}
+	Result->SetArrayField(TEXT("tags"), TagsArray);
+	
+	// Components
+	TArray<TSharedPtr<FJsonValue>> ComponentsArray;
+	TInlineComponentArray<UActorComponent*> Components;
+	Actor->GetComponents(Components);
+	for (UActorComponent* Component : Components)
+	{
+		TSharedPtr<FJsonObject> CompObj = MakeShared<FJsonObject>();
+		CompObj->SetStringField(TEXT("name"), Component->GetName());
+		CompObj->SetStringField(TEXT("class"), Component->GetClass()->GetName());
+		ComponentsArray.Add(MakeShared<FJsonValueObject>(CompObj));
+	}
+	Result->SetArrayField(TEXT("components"), ComponentsArray);
+	
+	return FECACommandResult::Success(Result);
+}
+
+//------------------------------------------------------------------------------
+// SetActorProperty
+//------------------------------------------------------------------------------
+
+FECACommandResult FECACommand_SetActorProperty::Execute(const TSharedPtr<FJsonObject>& Params)
+{
+	// Support both actor_name (legacy) and actor_path (new)
+	FString ActorName;
+	FString ActorPath;
+	AActor* Actor = nullptr;
+	
+	if (GetStringParam(Params, TEXT("actor_path"), ActorPath, false))
+	{
+		Actor = FindObject<AActor>(nullptr, *ActorPath);
+		if (!Actor)
+		{
+			Actor = LoadObject<AActor>(nullptr, *ActorPath);
+		}
+	}
+	else if (GetStringParam(Params, TEXT("actor_name"), ActorName, false))
+	{
+		Actor = FindActorByName(ActorName);
+	}
+	else
+	{
+		return FECACommandResult::Error(TEXT("Missing required parameter: actor_path or actor_name"));
+	}
+	
+	if (!Actor)
+	{
+		return FECACommandResult::Error(FString::Printf(TEXT("Actor not found: %s"), ActorPath.IsEmpty() ? *ActorName : *ActorPath));
+	}
+	
+	FString PropertyName;
+	if (!GetStringParam(Params, TEXT("property_name"), PropertyName))
+	{
+		return FECACommandResult::Error(TEXT("Missing required parameter: property_name"));
+	}
+	
+	// Optional: target a specific component
+	FString ComponentName;
+	UObject* TargetObject = Actor;
+	if (GetStringParam(Params, TEXT("component_name"), ComponentName, false))
+	{
+		UActorComponent* Component = Actor->FindComponentByClass<UActorComponent>();
+		for (UActorComponent* Comp : Actor->GetComponents())
+		{
+			if (Comp->GetName().Equals(ComponentName, ESearchCase::IgnoreCase) ||
+				Comp->GetClass()->GetName().Equals(ComponentName, ESearchCase::IgnoreCase))
+			{
+				TargetObject = Comp;
+				break;
+			}
+		}
+	}
+	
+	TSharedPtr<FJsonObject> Result = MakeResult();
+	Result->SetStringField(TEXT("actor_name"), Actor->GetActorLabel());
+	Result->SetStringField(TEXT("property_name"), PropertyName);
+	
+	// Handle special properties first
+	if (PropertyName.Equals(TEXT("Mobility"), ESearchCase::IgnoreCase))
+	{
+		FString Value;
+		if (GetStringParam(Params, TEXT("property_value"), Value))
+		{
+			USceneComponent* RootComponent = Actor->GetRootComponent();
+			if (RootComponent)
+			{
+				if (Value.Equals(TEXT("Static"), ESearchCase::IgnoreCase))
+				{
+					RootComponent->SetMobility(EComponentMobility::Static);
+				}
+				else if (Value.Equals(TEXT("Stationary"), ESearchCase::IgnoreCase))
+				{
+					RootComponent->SetMobility(EComponentMobility::Stationary);
+				}
+				else if (Value.Equals(TEXT("Movable"), ESearchCase::IgnoreCase))
+				{
+					RootComponent->SetMobility(EComponentMobility::Movable);
+				}
+				Result->SetStringField(TEXT("new_value"), Value);
+			}
+		}
+		return FECACommandResult::Success(Result);
+	}
+	
+	if (PropertyName.Equals(TEXT("Hidden"), ESearchCase::IgnoreCase))
+	{
+		bool bHidden;
+		if (GetBoolParam(Params, TEXT("property_value"), bHidden))
+		{
+			Actor->SetIsTemporarilyHiddenInEditor(bHidden);
+			Result->SetBoolField(TEXT("new_value"), bHidden);
+		}
+		return FECACommandResult::Success(Result);
+	}
+	
+	if (PropertyName.Equals(TEXT("Folder"), ESearchCase::IgnoreCase))
+	{
+		FString FolderPath;
+		if (GetStringParam(Params, TEXT("property_value"), FolderPath))
+		{
+			Actor->SetFolderPath(*FolderPath);
+			Result->SetStringField(TEXT("new_value"), FolderPath);
+		}
+		return FECACommandResult::Success(Result);
+	}
+	
+	// Handle StaticMesh property specially (common case)
+	if (PropertyName.Equals(TEXT("StaticMesh"), ESearchCase::IgnoreCase))
+	{
+		FString MeshPath;
+		if (!GetStringParam(Params, TEXT("property_value"), MeshPath))
+		{
+			return FECACommandResult::Error(TEXT("Missing property_value for StaticMesh"));
+		}
+		
+		UStaticMesh* NewMesh = LoadObject<UStaticMesh>(nullptr, *MeshPath);
+		if (!NewMesh)
+		{
+			return FECACommandResult::Error(FString::Printf(TEXT("Static mesh not found: %s"), *MeshPath));
+		}
+		
+		UStaticMeshComponent* MeshComp = nullptr;
+		if (AStaticMeshActor* SMActor = Cast<AStaticMeshActor>(Actor))
+		{
+			MeshComp = SMActor->GetStaticMeshComponent();
+		}
+		else
+		{
+			MeshComp = Actor->FindComponentByClass<UStaticMeshComponent>();
+		}
+		
+		if (!MeshComp)
+		{
+			return FECACommandResult::Error(TEXT("Actor has no StaticMeshComponent"));
+		}
+		
+		FString OldPath = MeshComp->GetStaticMesh() ? MeshComp->GetStaticMesh()->GetPathName() : TEXT("None");
+		MeshComp->SetStaticMesh(NewMesh);
+		
+		Result->SetStringField(TEXT("old_value"), OldPath);
+		Result->SetStringField(TEXT("new_value"), NewMesh->GetPathName());
+		return FECACommandResult::Success(Result);
+	}
+	
+	// Handle Material property
+	if (PropertyName.Equals(TEXT("Material"), ESearchCase::IgnoreCase) || PropertyName.StartsWith(TEXT("Material[")))
+	{
+		FString MaterialPath;
+		if (!GetStringParam(Params, TEXT("property_value"), MaterialPath))
+		{
+			return FECACommandResult::Error(TEXT("Missing property_value for Material"));
+		}
+		
+		int32 SlotIndex = 0;
+		if (PropertyName.StartsWith(TEXT("Material[")))
+		{
+			FString IndexStr = PropertyName.Mid(9);
+			IndexStr = IndexStr.LeftChop(1);  // Remove trailing ]
+			SlotIndex = FCString::Atoi(*IndexStr);
+		}
+		
+		UMaterialInterface* NewMaterial = LoadObject<UMaterialInterface>(nullptr, *MaterialPath);
+		if (!NewMaterial)
+		{
+			return FECACommandResult::Error(FString::Printf(TEXT("Material not found: %s"), *MaterialPath));
+		}
+		
+		UPrimitiveComponent* PrimComp = Actor->FindComponentByClass<UPrimitiveComponent>();
+		if (!PrimComp)
+		{
+			return FECACommandResult::Error(TEXT("Actor has no PrimitiveComponent"));
+		}
+		
+		UMaterialInterface* OldMaterial = PrimComp->GetMaterial(SlotIndex);
+		FString OldPath = OldMaterial ? OldMaterial->GetPathName() : TEXT("None");
+		PrimComp->SetMaterial(SlotIndex, NewMaterial);
+		
+		Result->SetStringField(TEXT("old_value"), OldPath);
+		Result->SetStringField(TEXT("new_value"), NewMaterial->GetPathName());
+		Result->SetNumberField(TEXT("slot_index"), SlotIndex);
+		return FECACommandResult::Success(Result);
+	}
+	
+	// Generic property setting via UE reflection
+	FProperty* Property = TargetObject->GetClass()->FindPropertyByName(*PropertyName);
+	if (!Property)
+	{
+		// Try case-insensitive search
+		for (TFieldIterator<FProperty> It(TargetObject->GetClass()); It; ++It)
+		{
+			if (It->GetName().Equals(PropertyName, ESearchCase::IgnoreCase))
+			{
+				Property = *It;
+				break;
+			}
+		}
+	}
+	
+	if (!Property)
+	{
+		return FECACommandResult::Error(FString::Printf(TEXT("Property not found: %s"), *PropertyName));
+	}
+	
+	void* ValuePtr = Property->ContainerPtrToValuePtr<void>(TargetObject);
+	
+	// Handle different property types
+	if (FNumericProperty* NumericProp = CastField<FNumericProperty>(Property))
+	{
+		double Value;
+		if (GetFloatParam(Params, TEXT("property_value"), Value))
+		{
+			if (NumericProp->IsInteger())
+			{
+				NumericProp->SetIntPropertyValue(ValuePtr, (int64)Value);
+			}
+			else
+			{
+				NumericProp->SetFloatingPointPropertyValue(ValuePtr, Value);
+			}
+			Result->SetNumberField(TEXT("new_value"), Value);
+		}
+	}
+	else if (FBoolProperty* BoolProp = CastField<FBoolProperty>(Property))
+	{
+		bool Value;
+		if (GetBoolParam(Params, TEXT("property_value"), Value))
+		{
+			BoolProp->SetPropertyValue(ValuePtr, Value);
+			Result->SetBoolField(TEXT("new_value"), Value);
+		}
+	}
+	else if (FStrProperty* StrProp = CastField<FStrProperty>(Property))
+	{
+		FString Value;
+		if (GetStringParam(Params, TEXT("property_value"), Value))
+		{
+			StrProp->SetPropertyValue(ValuePtr, Value);
+			Result->SetStringField(TEXT("new_value"), Value);
+		}
+	}
+	else if (FNameProperty* NameProp = CastField<FNameProperty>(Property))
+	{
+		FString Value;
+		if (GetStringParam(Params, TEXT("property_value"), Value))
+		{
+			NameProp->SetPropertyValue(ValuePtr, FName(*Value));
+			Result->SetStringField(TEXT("new_value"), Value);
+		}
+	}
+	else if (FObjectProperty* ObjProp = CastField<FObjectProperty>(Property))
+	{
+		FString ObjectPath;
+		if (GetStringParam(Params, TEXT("property_value"), ObjectPath))
+		{
+			UObject* NewObject = LoadObject<UObject>(nullptr, *ObjectPath);
+			if (NewObject)
+			{
+				ObjProp->SetObjectPropertyValue(ValuePtr, NewObject);
+				Result->SetStringField(TEXT("new_value"), NewObject->GetPathName());
+			}
+			else
+			{
+				return FECACommandResult::Error(FString::Printf(TEXT("Object not found: %s"), *ObjectPath));
+			}
+		}
+	}
+	else if (FStructProperty* StructProp = CastField<FStructProperty>(Property))
+	{
+		// Handle common struct types
+		if (StructProp->Struct == TBaseStructure<FVector>::Get())
+		{
+			FVector Value;
+			if (GetVectorParam(Params, TEXT("property_value"), Value))
+			{
+				*StructProp->ContainerPtrToValuePtr<FVector>(TargetObject) = Value;
+				Result->SetObjectField(TEXT("new_value"), VectorToJson(Value));
+			}
+		}
+		else if (StructProp->Struct == TBaseStructure<FRotator>::Get())
+		{
+			FRotator Value;
+			if (GetRotatorParam(Params, TEXT("property_value"), Value))
+			{
+				*StructProp->ContainerPtrToValuePtr<FRotator>(TargetObject) = Value;
+				Result->SetObjectField(TEXT("new_value"), RotatorToJson(Value));
+			}
+		}
+		else if (StructProp->Struct == TBaseStructure<FLinearColor>::Get())
+		{
+			FVector Value;
+			if (GetVectorParam(Params, TEXT("property_value"), Value))
+			{
+				FLinearColor Color(Value.X, Value.Y, Value.Z, 1.0f);
+				*StructProp->ContainerPtrToValuePtr<FLinearColor>(TargetObject) = Color;
+			}
+		}
+		else
+		{
+			return FECACommandResult::Error(FString::Printf(TEXT("Struct type not supported: %s"), *StructProp->Struct->GetName()));
+		}
+	}
+	else
+	{
+		return FECACommandResult::Error(FString::Printf(TEXT("Property type not supported: %s"), *Property->GetClass()->GetName()));
+	}
+	
+	// Notify that property changed
+	TargetObject->PostEditChange();
+	
+	return FECACommandResult::Success(Result);
+}
+
+//------------------------------------------------------------------------------
+// FindActors
+//------------------------------------------------------------------------------
+
+FECACommandResult FECACommand_FindActors::Execute(const TSharedPtr<FJsonObject>& Params)
+{
+	UWorld* World = GetEditorWorld();
+	if (!World)
+	{
+		return FECACommandResult::Error(TEXT("No world available"));
+	}
+	
+	FString NamePattern;
+	FString ClassName;
+	FString Tag;
+	FVector NearLocation;
+	double Radius = 1000.0;
+	
+	bool bHasNamePattern = GetStringParam(Params, TEXT("name_pattern"), NamePattern, false);
+	bool bHasClassName = GetStringParam(Params, TEXT("class_name"), ClassName, false);
+	bool bHasTag = GetStringParam(Params, TEXT("tag"), Tag, false);
+	bool bHasLocation = GetVectorParam(Params, TEXT("near_location"), NearLocation, false);
+	GetFloatParam(Params, TEXT("radius"), Radius, false);
+	
+	TArray<TSharedPtr<FJsonValue>> ActorsArray;
+	
+	for (TActorIterator<AActor> It(World); It; ++It)
+	{
+		AActor* Actor = *It;
+		bool bMatch = true;
+		
+		// Name pattern filter - check both label and object name
+		if (bHasNamePattern && bMatch)
+		{
+			FString ActorLabel = Actor->GetActorLabel();
+			FString ActorObjName = Actor->GetName();
+			bMatch = ActorLabel.MatchesWildcard(NamePattern, ESearchCase::IgnoreCase) || 
+			         ActorObjName.MatchesWildcard(NamePattern, ESearchCase::IgnoreCase);
+		}
+		
+		// Class filter
+		if (bHasClassName && bMatch)
+		{
+			FString ActorClassName = Actor->GetClass()->GetName();
+			bMatch = ActorClassName.Contains(ClassName);
+		}
+		
+		// Tag filter
+		if (bHasTag && bMatch)
+		{
+			bMatch = Actor->Tags.Contains(FName(*Tag));
+		}
+		
+		// Location filter
+		if (bHasLocation && bMatch)
+		{
+			double Distance = FVector::Dist(Actor->GetActorLocation(), NearLocation);
+			bMatch = Distance <= Radius;
+		}
+		
+		if (bMatch)
+		{
+			TSharedPtr<FJsonObject> ActorJson = MakeShared<FJsonObject>();
+			ActorJson->SetStringField(TEXT("name"), Actor->GetActorLabel());
+			ActorJson->SetStringField(TEXT("class"), Actor->GetClass()->GetName());
+			ActorJson->SetObjectField(TEXT("location"), VectorToJson(Actor->GetActorLocation()));
+			ActorsArray.Add(MakeShared<FJsonValueObject>(ActorJson));
+		}
+	}
+	
+	TSharedPtr<FJsonObject> Result = MakeResult();
+	Result->SetArrayField(TEXT("actors"), ActorsArray);
+	Result->SetNumberField(TEXT("count"), ActorsArray.Num());
+	return FECACommandResult::Success(Result);
+}
+
+//------------------------------------------------------------------------------
+// DuplicateActor
+//------------------------------------------------------------------------------
+
+FECACommandResult FECACommand_DuplicateActor::Execute(const TSharedPtr<FJsonObject>& Params)
+{
+	FString ActorName;
+	if (!GetStringParam(Params, TEXT("actor_name"), ActorName))
+	{
+		return FECACommandResult::Error(TEXT("Missing required parameter: actor_name"));
+	}
+	
+	AActor* SourceActor = FindActorByName(ActorName);
+	if (!SourceActor)
+	{
+		return FECACommandResult::Error(FString::Printf(TEXT("Actor not found: %s"), *ActorName));
+	}
+	
+	UWorld* World = GetEditorWorld();
+	if (!World)
+	{
+		return FECACommandResult::Error(TEXT("No world available"));
+	}
+	
+	// Get offset
+	FVector Offset = FVector(100, 0, 0); // Default offset
+	GetVectorParam(Params, TEXT("offset"), Offset, false);
+	
+	// Duplicate using editor subsystem
+	UEditorActorSubsystem* EditorActorSubsystem = GEditor->GetEditorSubsystem<UEditorActorSubsystem>();
+	if (!EditorActorSubsystem)
+	{
+		return FECACommandResult::Error(TEXT("EditorActorSubsystem not available"));
+	}
+	
+	AActor* NewActor = EditorActorSubsystem->DuplicateActor(SourceActor, World);
+	if (!NewActor)
+	{
+		return FECACommandResult::Error(TEXT("Failed to duplicate actor"));
+	}
+	
+	// Apply offset
+	NewActor->SetActorLocation(SourceActor->GetActorLocation() + Offset);
+	
+	// Set new name if specified
+	FString NewName;
+	if (GetStringParam(Params, TEXT("new_name"), NewName, false))
+	{
+		NewActor->SetActorLabel(*NewName);
+	}
+	
+	TSharedPtr<FJsonObject> Result = MakeResult();
+	Result->SetStringField(TEXT("actor_name"), NewActor->GetActorLabel());
+	Result->SetStringField(TEXT("internal_name"), NewActor->GetName());
+	Result->SetObjectField(TEXT("transform"), TransformToJson(NewActor->GetTransform()));
+	return FECACommandResult::Success(Result);
+}
+
+
