ECABridge experimental MCP plugin

diff --git a/Engine/Plugins/Experimental/ECABridge/Source/ECABridge/Private/Commands/ECAAssetCommands.cpp b/Engine/Plugins/Experimental/ECABridge/Source/ECABridge/Private/Commands/ECAAssetCommands.cpp
new file mode 100644
--- /dev/null
+++ b/Engine/Plugins/Experimental/ECABridge/Source/ECABridge/Private/Commands/ECAAssetCommands.cpp
@@ -0,0 +1,3781 @@
+// Copyright Epic Games, Inc. All Rights Reserved.
+
+#include "Commands/ECAAssetCommands.h"
+#include "Commands/ECACommand.h"
+#include "Editor.h"
+#include "Engine/Texture2D.h"
+#include "Materials/Material.h"
+#include "Materials/MaterialInstance.h"
+#include "Materials/MaterialInstanceConstant.h"
+#include "Materials/MaterialExpressionTextureSample.h"
+#include "Materials/MaterialExpressionConstant.h"
+#include "Materials/MaterialExpressionConstant3Vector.h"
+#include "Materials/MaterialExpressionConstant4Vector.h"
+#include "Materials/MaterialExpressionVectorParameter.h"
+#include "Materials/MaterialExpressionScalarParameter.h"
+#include "Materials/MaterialExpressionTextureCoordinate.h"
+#include "Materials/MaterialExpressionTextureSampleParameter2D.h"
+#include "AssetRegistry/AssetRegistryModule.h"
+#include "AssetToolsModule.h"
+#include "IAssetTools.h"
+#include "ObjectTools.h"
+#include "Factories/TextureFactory.h"
+#include "Factories/MaterialFactoryNew.h"
+#include "Factories/MaterialInstanceConstantFactoryNew.h"
+#include "Engine/StaticMesh.h"
+#include "Engine/SkeletalMesh.h"
+#include "Components/StaticMeshComponent.h"
+#include "Components/SkeletalMeshComponent.h"
+#include "Components/PrimitiveComponent.h"
+#include "Engine/StaticMeshActor.h"
+#include "GameFramework/Actor.h"
+#include "EngineUtils.h"
+#include "HAL/FileManager.h"
+#include "Misc/Paths.h"
+#include "Misc/FileHelper.h"
+#include "UObject/SavePackage.h"
+#include "MeshDescription.h"
+#include "StaticMeshAttributes.h"
+#include "Engine/StaticMesh.h"
+#include "PhysicsEngine/BodySetup.h"
+#include "Engine/TextureRenderTarget2D.h"
+#include "TextureResource.h"
+#include "ImageUtils.h"
+#include "InterchangeManager.h"
+#include "InterchangeSourceData.h"
+#include "InterchangeAssetImportData.h"
+#include "StaticMeshResources.h"
+#include "ThumbnailRendering/ThumbnailManager.h"
+#include "Misc/Base64.h"
+#include "CanvasTypes.h"
+#include "CanvasItem.h"
+
+// Register all asset commands
+REGISTER_ECA_COMMAND(FECACommand_ImportTexture)
+REGISTER_ECA_COMMAND(FECACommand_CreateMaterial)
+REGISTER_ECA_COMMAND(FECACommand_CreateMaterialFromTextures)
+REGISTER_ECA_COMMAND(FECACommand_GetTextureInfo)
+REGISTER_ECA_COMMAND(FECACommand_GetMaterialInfo)
+REGISTER_ECA_COMMAND(FECACommand_ListTextures)
+REGISTER_ECA_COMMAND(FECACommand_ListMaterials)
+REGISTER_ECA_COMMAND(FECACommand_SetMaterialTextureParam)
+REGISTER_ECA_COMMAND(FECACommand_SetActorMaterial)
+REGISTER_ECA_COMMAND(FECACommand_GetActorMaterials)
+REGISTER_ECA_COMMAND(FECACommand_SetStaticMeshMaterial)
+REGISTER_ECA_COMMAND(FECACommand_CreateMaterialInstance)
+REGISTER_ECA_COMMAND(FECACommand_SetMaterialInstanceScalarParam)
+REGISTER_ECA_COMMAND(FECACommand_SetMaterialInstanceVectorParam)
+REGISTER_ECA_COMMAND(FECACommand_SetMaterialInstanceTextureParam)
+REGISTER_ECA_COMMAND(FECACommand_ImportOBJ)
+REGISTER_ECA_COMMAND(FECACommand_ImportFBX)
+REGISTER_ECA_COMMAND(FECACommand_ExportTexture)
+REGISTER_ECA_COMMAND(FECACommand_ExportRenderTarget)
+REGISTER_ECA_COMMAND(FECACommand_SetAssetProperty)
+REGISTER_ECA_COMMAND(FECACommand_GetAssetProperty)
+REGISTER_ECA_COMMAND(FECACommand_ListAssetProperties)
+REGISTER_ECA_COMMAND(FECACommand_DeleteAsset)
+REGISTER_ECA_COMMAND(FECACommand_RenameAsset)
+REGISTER_ECA_COMMAND(FECACommand_MoveAsset)
+REGISTER_ECA_COMMAND(FECACommand_DuplicateAsset)
+REGISTER_ECA_COMMAND(FECACommand_GetAssetThumbnail)
+REGISTER_ECA_COMMAND(FECACommand_GetAssetThumbnails)
+
+//------------------------------------------------------------------------------
+// Helper functions
+//------------------------------------------------------------------------------
+
+// Collect compilation errors from material expressions
+static TArray<TSharedPtr<FJsonValue>> GetMaterialCompilationErrors(UMaterial* Material)
+{
+	TArray<TSharedPtr<FJsonValue>> ErrorsArray;
+	
+	if (!Material || !Material->GetEditorOnlyData())
+	{
+		return ErrorsArray;
+	}
+	
+	for (UMaterialExpression* Expr : Material->GetEditorOnlyData()->ExpressionCollection.Expressions)
+	{
+		if (Expr && !Expr->LastErrorText.IsEmpty())
+		{
+			TSharedPtr<FJsonObject> ErrorObj = MakeShared<FJsonObject>();
+			ErrorObj->SetStringField(TEXT("node_id"), Expr->MaterialExpressionGuid.ToString());
+			ErrorObj->SetStringField(TEXT("node_type"), Expr->GetClass()->GetName().Replace(TEXT("MaterialExpression"), TEXT("")));
+			ErrorObj->SetStringField(TEXT("error"), Expr->LastErrorText);
+			ErrorsArray.Add(MakeShared<FJsonValueObject>(ErrorObj));
+		}
+	}
+	
+	return ErrorsArray;
+}
+
+// Add compilation errors to result if any exist
+static void AddMaterialCompilationErrorsToResult(TSharedPtr<FJsonObject>& Result, UMaterial* Material)
+{
+	TArray<TSharedPtr<FJsonValue>> Errors = GetMaterialCompilationErrors(Material);
+	if (Errors.Num() > 0)
+	{
+		Result->SetArrayField(TEXT("compilation_errors"), Errors);
+		Result->SetBoolField(TEXT("has_errors"), true);
+	}
+	else
+	{
+		Result->SetBoolField(TEXT("has_errors"), false);
+	}
+}
+
+static TextureCompressionSettings GetCompressionSettingsFromString(const FString& Setting)
+{
+	if (Setting.Equals(TEXT("Normalmap"), ESearchCase::IgnoreCase))
+	{
+		return TC_Normalmap;
+	}
+	else if (Setting.Equals(TEXT("Masks"), ESearchCase::IgnoreCase))
+	{
+		return TC_Masks;
+	}
+	else if (Setting.Equals(TEXT("Grayscale"), ESearchCase::IgnoreCase))
+	{
+		return TC_Grayscale;
+	}
+	else if (Setting.Equals(TEXT("HDR"), ESearchCase::IgnoreCase))
+	{
+		return TC_HDR;
+	}
+	else if (Setting.Equals(TEXT("UserInterface2D"), ESearchCase::IgnoreCase) || Setting.Equals(TEXT("UI"), ESearchCase::IgnoreCase))
+	{
+		return TC_EditorIcon;
+	}
+	else if (Setting.Equals(TEXT("Alpha"), ESearchCase::IgnoreCase))
+	{
+		return TC_Alpha;
+	}
+	else if (Setting.Equals(TEXT("BC7"), ESearchCase::IgnoreCase))
+	{
+		return TC_BC7;
+	}
+	return TC_Default;
+}
+
+static TextureGroup GetLODGroupFromString(const FString& Group)
+{
+	if (Group.Equals(TEXT("WorldNormalMap"), ESearchCase::IgnoreCase))
+	{
+		return TEXTUREGROUP_WorldNormalMap;
+	}
+	else if (Group.Equals(TEXT("WorldSpecular"), ESearchCase::IgnoreCase))
+	{
+		return TEXTUREGROUP_WorldSpecular;
+	}
+	else if (Group.Equals(TEXT("UI"), ESearchCase::IgnoreCase))
+	{
+		return TEXTUREGROUP_UI;
+	}
+	else if (Group.Equals(TEXT("Shadowmap"), ESearchCase::IgnoreCase))
+	{
+		return TEXTUREGROUP_Shadowmap;
+	}
+	else if (Group.Equals(TEXT("Character"), ESearchCase::IgnoreCase))
+	{
+		return TEXTUREGROUP_Character;
+	}
+	else if (Group.Equals(TEXT("CharacterNormalMap"), ESearchCase::IgnoreCase))
+	{
+		return TEXTUREGROUP_CharacterNormalMap;
+	}
+	else if (Group.Equals(TEXT("CharacterSpecular"), ESearchCase::IgnoreCase))
+	{
+		return TEXTUREGROUP_CharacterSpecular;
+	}
+	else if (Group.Equals(TEXT("Effects"), ESearchCase::IgnoreCase))
+	{
+		return TEXTUREGROUP_Effects;
+	}
+	else if (Group.Equals(TEXT("Weapon"), ESearchCase::IgnoreCase))
+	{
+		return TEXTUREGROUP_Weapon;
+	}
+	else if (Group.Equals(TEXT("WeaponNormalMap"), ESearchCase::IgnoreCase))
+	{
+		return TEXTUREGROUP_WeaponNormalMap;
+	}
+	else if (Group.Equals(TEXT("WeaponSpecular"), ESearchCase::IgnoreCase))
+	{
+		return TEXTUREGROUP_WeaponSpecular;
+	}
+	else if (Group.Equals(TEXT("Vehicle"), ESearchCase::IgnoreCase))
+	{
+		return TEXTUREGROUP_Vehicle;
+	}
+	else if (Group.Equals(TEXT("VehicleNormalMap"), ESearchCase::IgnoreCase))
+	{
+		return TEXTUREGROUP_VehicleNormalMap;
+	}
+	else if (Group.Equals(TEXT("VehicleSpecular"), ESearchCase::IgnoreCase))
+	{
+		return TEXTUREGROUP_VehicleSpecular;
+	}
+	else if (Group.Equals(TEXT("Skybox"), ESearchCase::IgnoreCase))
+	{
+		return TEXTUREGROUP_Skybox;
+	}
+	return TEXTUREGROUP_World;
+}
+
+static EBlendMode GetBlendModeFromString(const FString& Mode)
+{
+	if (Mode.Equals(TEXT("Masked"), ESearchCase::IgnoreCase))
+	{
+		return BLEND_Masked;
+	}
+	else if (Mode.Equals(TEXT("Translucent"), ESearchCase::IgnoreCase))
+	{
+		return BLEND_Translucent;
+	}
+	else if (Mode.Equals(TEXT("Additive"), ESearchCase::IgnoreCase))
+	{
+		return BLEND_Additive;
+	}
+	else if (Mode.Equals(TEXT("Modulate"), ESearchCase::IgnoreCase))
+	{
+		return BLEND_Modulate;
+	}
+	return BLEND_Opaque;
+}
+
+struct FTextureImportSettings
+{
+	TextureCompressionSettings CompressionSettings = TC_Default;
+	TextureGroup LODGroup = TEXTUREGROUP_World;
+	bool bSRGB = true;
+};
+
+static UTexture2D* ImportTextureFromFile(const FString& SourcePath, const FString& DestinationPath, const FString& TextureName, FString& OutError, const FTextureImportSettings& Settings = FTextureImportSettings())
+{
+	// Check if source file exists
+	if (!FPaths::FileExists(SourcePath))
+	{
+		OutError = FString::Printf(TEXT("Source file does not exist: %s"), *SourcePath);
+		return nullptr;
+	}
+	
+	// Create the texture factory
+	UTextureFactory* TextureFactory = NewObject<UTextureFactory>();
+	TextureFactory->AddToRoot(); // Prevent garbage collection during import
+	
+	// Configure factory settings BEFORE import - this is critical for gamma space
+	TextureFactory->SuppressImportOverwriteDialog();
+	TextureFactory->NoCompression = false;
+	TextureFactory->CompressionSettings = Settings.CompressionSettings;
+	TextureFactory->bRGBToEmissive = false;
+	TextureFactory->bRGBToBaseColor = false;
+	TextureFactory->LODGroup = Settings.LODGroup;
+	
+	// Set color space on factory - this affects how source data is interpreted
+	if (Settings.CompressionSettings == TC_Normalmap)
+	{
+		TextureFactory->bDoScaleMipsForAlphaCoverage = false;
+	}
+	
+	// Determine the asset name
+	FString AssetName = TextureName;
+	if (AssetName.IsEmpty())
+	{
+		AssetName = FPaths::GetBaseFilename(SourcePath);
+	}
+	
+	// Ensure destination path ends with /
+	FString DestPath = DestinationPath;
+	if (!DestPath.EndsWith(TEXT("/")))
+	{
+		DestPath += TEXT("/");
+	}
+	
+	// Full package path
+	FString PackagePath = DestPath + AssetName;
+	
+	// Create the package
+	UPackage* Package = CreatePackage(*PackagePath);
+	if (!Package)
+	{
+		OutError = TEXT("Failed to create package");
+		TextureFactory->RemoveFromRoot();
+		return nullptr;
+	}
+	
+	// Read the file data
+	TArray<uint8> FileData;
+	if (!FFileHelper::LoadFileToArray(FileData, *SourcePath))
+	{
+		OutError = FString::Printf(TEXT("Failed to read source file: %s"), *SourcePath);
+		TextureFactory->RemoveFromRoot();
+		return nullptr;
+	}
+	
+	// Import the texture
+	const uint8* DataPtr = FileData.GetData();
+	UTexture2D* ImportedTexture = Cast<UTexture2D>(TextureFactory->FactoryCreateBinary(
+		UTexture2D::StaticClass(),
+		Package,
+		*AssetName,
+		RF_Public | RF_Standalone,
+		nullptr,
+		*FPaths::GetExtension(SourcePath),
+		DataPtr,
+		DataPtr + FileData.Num(),
+		GWarn
+	));
+	
+	TextureFactory->RemoveFromRoot();
+	
+	if (!ImportedTexture)
+	{
+		OutError = TEXT("Failed to import texture");
+		return nullptr;
+	}
+	
+	// Apply settings that weren't set by factory - wrap in PreEdit/PostEdit scope
+	// to ensure derived data is properly regenerated with correct gamma space
+	ImportedTexture->PreEditChange(nullptr);
+	ImportedTexture->SRGB = Settings.bSRGB;
+	ImportedTexture->CompressionSettings = Settings.CompressionSettings;
+	ImportedTexture->LODGroup = Settings.LODGroup;
+	ImportedTexture->PostEditChange();
+	
+	// Notify asset registry
+	FAssetRegistryModule::AssetCreated(ImportedTexture);
+	Package->MarkPackageDirty();
+	
+	return ImportedTexture;
+}
+
+static UTexture2D* LoadTextureByPath(const FString& TexturePath)
+{
+	return LoadObject<UTexture2D>(nullptr, *TexturePath);
+}
+
+static UMaterial* LoadMaterialByPath(const FString& MaterialPath)
+{
+	return LoadObject<UMaterial>(nullptr, *MaterialPath);
+}
+
+//------------------------------------------------------------------------------
+// ImportTexture
+//------------------------------------------------------------------------------
+
+FECACommandResult FECACommand_ImportTexture::Execute(const TSharedPtr<FJsonObject>& Params)
+{
+	FString SourcePath;
+	if (!GetStringParam(Params, TEXT("source_path"), SourcePath))
+	{
+		return FECACommandResult::Error(TEXT("Missing required parameter: source_path"));
+	}
+	
+	FString DestinationPath;
+	if (!GetStringParam(Params, TEXT("destination_path"), DestinationPath))
+	{
+		return FECACommandResult::Error(TEXT("Missing required parameter: destination_path"));
+	}
+	
+	FString TextureName;
+	GetStringParam(Params, TEXT("texture_name"), TextureName, false);
+	
+	FString CompressionSettingsStr = TEXT("Default");
+	GetStringParam(Params, TEXT("compression_settings"), CompressionSettingsStr, false);
+	
+	bool bSRGB = true;
+	GetBoolParam(Params, TEXT("srgb"), bSRGB, false);
+	
+	FString LODGroupStr = TEXT("World");
+	GetStringParam(Params, TEXT("lod_group"), LODGroupStr, false);
+	
+	// Configure import settings BEFORE import to ensure correct gamma space
+	FTextureImportSettings ImportSettings;
+	ImportSettings.CompressionSettings = GetCompressionSettingsFromString(CompressionSettingsStr);
+	ImportSettings.bSRGB = bSRGB;
+	ImportSettings.LODGroup = GetLODGroupFromString(LODGroupStr);
+	
+	// Import the texture with settings
+	FString Error;
+	UTexture2D* ImportedTexture = ImportTextureFromFile(SourcePath, DestinationPath, TextureName, Error, ImportSettings);
+	
+	if (!ImportedTexture)
+	{
+		return FECACommandResult::Error(Error);
+	}
+	
+	TSharedPtr<FJsonObject> Result = MakeResult();
+	Result->SetStringField(TEXT("texture_path"), ImportedTexture->GetPathName());
+	Result->SetStringField(TEXT("texture_name"), ImportedTexture->GetName());
+	Result->SetNumberField(TEXT("width"), ImportedTexture->GetSizeX());
+	Result->SetNumberField(TEXT("height"), ImportedTexture->GetSizeY());
+	Result->SetStringField(TEXT("compression_settings"), CompressionSettingsStr);
+	Result->SetBoolField(TEXT("srgb"), bSRGB);
+	
+	return FECACommandResult::Success(Result);
+}
+
+//------------------------------------------------------------------------------
+// CreateMaterial
+//------------------------------------------------------------------------------
+
+FECACommandResult FECACommand_CreateMaterial::Execute(const TSharedPtr<FJsonObject>& Params)
+{
+	FString MaterialName;
+	if (!GetStringParam(Params, TEXT("material_name"), MaterialName))
+	{
+		return FECACommandResult::Error(TEXT("Missing required parameter: material_name"));
+	}
+	
+	FString Path = TEXT("/Game/Materials/");
+	GetStringParam(Params, TEXT("path"), Path, false);
+	
+	// Ensure path ends with /
+	if (!Path.EndsWith(TEXT("/")))
+	{
+		Path += TEXT("/");
+	}
+	
+	// Get optional texture paths
+	FString BaseColorTexturePath, NormalTexturePath, RoughnessTexturePath, MetallicTexturePath, EmissiveTexturePath;
+	GetStringParam(Params, TEXT("base_color_texture"), BaseColorTexturePath, false);
+	GetStringParam(Params, TEXT("normal_texture"), NormalTexturePath, false);
+	GetStringParam(Params, TEXT("roughness_texture"), RoughnessTexturePath, false);
+	GetStringParam(Params, TEXT("metallic_texture"), MetallicTexturePath, false);
+	GetStringParam(Params, TEXT("emissive_texture"), EmissiveTexturePath, false);
+	
+	// Get scalar values
+	double Roughness = 0.5;
+	double Metallic = 0.0;
+	GetFloatParam(Params, TEXT("roughness"), Roughness, false);
+	GetFloatParam(Params, TEXT("metallic"), Metallic, false);
+	
+	bool bTwoSided = false;
+	GetBoolParam(Params, TEXT("two_sided"), bTwoSided, false);
+	
+	FString BlendModeStr = TEXT("Opaque");
+	GetStringParam(Params, TEXT("blend_mode"), BlendModeStr, false);
+	
+	// Create package path
+	FString PackagePath = Path + MaterialName;
+	
+	// Check if asset already exists at this path
+	if (UObject* ExistingAsset = StaticLoadObject(UObject::StaticClass(), nullptr, *PackagePath))
+	{
+		if (UMaterial* ExistingMaterial = Cast<UMaterial>(ExistingAsset))
+		{
+			TSharedPtr<FJsonObject> Result = MakeResult();
+			Result->SetStringField(TEXT("material_path"), PackagePath);
+			Result->SetStringField(TEXT("material_name"), MaterialName);
+			Result->SetBoolField(TEXT("already_exists"), true);
+			Result->SetStringField(TEXT("message"), TEXT("Material already exists at this path"));
+			return FECACommandResult::Success(Result);
+		}
+		else
+		{
+			return FECACommandResult::Error(FString::Printf(TEXT("Asset already exists at path '%s' but is not a Material (it's a %s)"), *PackagePath, *ExistingAsset->GetClass()->GetName()));
+		}
+	}
+	
+	// Create package
+	UPackage* Package = CreatePackage(*PackagePath);
+	if (!Package)
+	{
+		return FECACommandResult::Error(TEXT("Failed to create package"));
+	}
+	
+	// Create material factory
+	UMaterialFactoryNew* MaterialFactory = NewObject<UMaterialFactoryNew>();
+	
+	// Create the material
+	UMaterial* NewMaterial = Cast<UMaterial>(MaterialFactory->FactoryCreateNew(
+		UMaterial::StaticClass(),
+		Package,
+		*MaterialName,
+		RF_Public | RF_Standalone,
+		nullptr,
+		GWarn
+	));
+	
+	if (!NewMaterial)
+	{
+		return FECACommandResult::Error(TEXT("Failed to create material"));
+	}
+	
+	// Set material properties
+	NewMaterial->TwoSided = bTwoSided;
+	NewMaterial->BlendMode = GetBlendModeFromString(BlendModeStr);
+	
+	// Track position for material expression nodes
+	int32 NodePosX = -400;
+	int32 NodePosY = 0;
+	const int32 NodeSpacingY = 200;
+	
+	// Helper lambda to create texture sampler and connect to material
+	auto AddTextureSampler = [&](const FString& TexturePath, int32 OutputIndex, EMaterialProperty Property) -> bool
+	{
+		if (TexturePath.IsEmpty())
+		{
+			return false;
+		}
+		
+		UTexture2D* Texture = LoadTextureByPath(TexturePath);
+		if (!Texture)
+		{
+			return false;
+		}
+		
+		UMaterialExpressionTextureSampleParameter2D* TextureSampler = NewObject<UMaterialExpressionTextureSampleParameter2D>(NewMaterial);
+		TextureSampler->Texture = Texture;
+		TextureSampler->ParameterName = FName(*FString::Printf(TEXT("%s_Texture"), *UEnum::GetValueAsString(Property)));
+		
+		// Set appropriate sampler type based on material property and texture format
+		if (Property == MP_Normal)
+		{
+			TextureSampler->SamplerType = SAMPLERTYPE_Normal;
+		}
+		else if (Property == MP_Roughness || Property == MP_Metallic || Property == MP_AmbientOcclusion)
+		{
+			// Use LinearGrayscale for single-channel textures, Masks for packed RGB textures
+			if (Texture->CompressionSettings == TC_Grayscale || Texture->CompressionSettings == TC_Alpha)
+			{
+				TextureSampler->SamplerType = SAMPLERTYPE_LinearGrayscale;
+			}
+			else
+			{
+				TextureSampler->SamplerType = SAMPLERTYPE_Masks;
+			}
+		}
+		else
+		{
+			TextureSampler->SamplerType = SAMPLERTYPE_Color;
+		}
+		
+		TextureSampler->MaterialExpressionEditorX = NodePosX;
+		TextureSampler->MaterialExpressionEditorY = NodePosY;
+		NodePosY += NodeSpacingY;
+		
+		NewMaterial->GetEditorOnlyData()->ExpressionCollection.Expressions.Add(TextureSampler);
+		
+		// Get the material input for this property
+		FExpressionInput* Input = NewMaterial->GetExpressionInputForProperty(Property);
+		if (Input)
+		{
+			Input->Expression = TextureSampler;
+			Input->OutputIndex = OutputIndex;
+		}
+		
+		return true;
+	};
+	
+	// Add base color
+	if (!BaseColorTexturePath.IsEmpty())
+	{
+		AddTextureSampler(BaseColorTexturePath, 0, MP_BaseColor);
+	}
+	else
+	{
+		// Check for base color value
+		const TSharedPtr<FJsonObject>* BaseColorObj = NULL;
+		if (GetObjectParam(Params, TEXT("base_color"), BaseColorObj, false) && BaseColorObj)
+		{
+			UMaterialExpressionVectorParameter* ColorParam = NewObject<UMaterialExpressionVectorParameter>(NewMaterial);
+			ColorParam->ParameterName = TEXT("BaseColor");
+			ColorParam->DefaultValue.R = (*BaseColorObj)->GetNumberField(TEXT("r"));
+			ColorParam->DefaultValue.G = (*BaseColorObj)->GetNumberField(TEXT("g"));
+			ColorParam->DefaultValue.B = (*BaseColorObj)->GetNumberField(TEXT("b"));
+			ColorParam->DefaultValue.A = (*BaseColorObj)->HasField(TEXT("a")) ? (*BaseColorObj)->GetNumberField(TEXT("a")) : 1.0f;
+			ColorParam->MaterialExpressionEditorX = NodePosX;
+			ColorParam->MaterialExpressionEditorY = NodePosY;
+			NodePosY += NodeSpacingY;
+			
+			NewMaterial->GetEditorOnlyData()->ExpressionCollection.Expressions.Add(ColorParam);
+			NewMaterial->GetEditorOnlyData()->BaseColor.Expression = ColorParam;
+		}
+	}
+	
+	// Add normal map
+	if (!NormalTexturePath.IsEmpty())
+	{
+		AddTextureSampler(NormalTexturePath, 0, MP_Normal);
+	}
+	
+	// Add roughness
+	if (!RoughnessTexturePath.IsEmpty())
+	{
+		AddTextureSampler(RoughnessTexturePath, 0, MP_Roughness);
+	}
+	else
+	{
+		UMaterialExpressionScalarParameter* RoughnessParam = NewObject<UMaterialExpressionScalarParameter>(NewMaterial);
+		RoughnessParam->ParameterName = TEXT("Roughness");
+		RoughnessParam->DefaultValue = Roughness;
+		RoughnessParam->MaterialExpressionEditorX = NodePosX;
+		RoughnessParam->MaterialExpressionEditorY = NodePosY;
+		NodePosY += NodeSpacingY;
+		
+		NewMaterial->GetEditorOnlyData()->ExpressionCollection.Expressions.Add(RoughnessParam);
+		NewMaterial->GetEditorOnlyData()->Roughness.Expression = RoughnessParam;
+	}
+	
+	// Add metallic
+	if (!MetallicTexturePath.IsEmpty())
+	{
+		AddTextureSampler(MetallicTexturePath, 0, MP_Metallic);
+	}
+	else
+	{
+		UMaterialExpressionScalarParameter* MetallicParam = NewObject<UMaterialExpressionScalarParameter>(NewMaterial);
+		MetallicParam->ParameterName = TEXT("Metallic");
+		MetallicParam->DefaultValue = Metallic;
+		MetallicParam->MaterialExpressionEditorX = NodePosX;
+		MetallicParam->MaterialExpressionEditorY = NodePosY;
+		NodePosY += NodeSpacingY;
+		
+		NewMaterial->GetEditorOnlyData()->ExpressionCollection.Expressions.Add(MetallicParam);
+		NewMaterial->GetEditorOnlyData()->Metallic.Expression = MetallicParam;
+	}
+	
+	// Add emissive
+	if (!EmissiveTexturePath.IsEmpty())
+	{
+		AddTextureSampler(EmissiveTexturePath, 0, MP_EmissiveColor);
+	}
+	
+	// Compile the material
+	NewMaterial->PreEditChange(nullptr);
+	NewMaterial->PostEditChange();
+	
+	// Notify asset registry
+	FAssetRegistryModule::AssetCreated(NewMaterial);
+	Package->MarkPackageDirty();
+	
+	TSharedPtr<FJsonObject> Result = MakeResult();
+	Result->SetStringField(TEXT("material_path"), NewMaterial->GetPathName());
+	Result->SetStringField(TEXT("material_name"), NewMaterial->GetName());
+	Result->SetBoolField(TEXT("two_sided"), bTwoSided);
+	Result->SetStringField(TEXT("blend_mode"), BlendModeStr);
+	
+	// Include any compilation errors
+	AddMaterialCompilationErrorsToResult(Result, NewMaterial);
+	
+	return FECACommandResult::Success(Result);
+}
+
+//------------------------------------------------------------------------------
+// CreateMaterialFromTextures
+//------------------------------------------------------------------------------
+
+FECACommandResult FECACommand_CreateMaterialFromTextures::Execute(const TSharedPtr<FJsonObject>& Params)
+{
+	FString MaterialName;
+	if (!GetStringParam(Params, TEXT("material_name"), MaterialName))
+	{
+		return FECACommandResult::Error(TEXT("Missing required parameter: material_name"));
+	}
+	
+	FString MaterialPath = TEXT("/Game/Materials/");
+	FString TexturePath = TEXT("/Game/Textures/");
+	GetStringParam(Params, TEXT("material_path"), MaterialPath, false);
+	GetStringParam(Params, TEXT("texture_path"), TexturePath, false);
+	
+	// Ensure paths end with /
+	if (!MaterialPath.EndsWith(TEXT("/")))
+	{
+		MaterialPath += TEXT("/");
+	}
+	if (!TexturePath.EndsWith(TEXT("/")))
+	{
+		TexturePath += TEXT("/");
+	}
+	
+	// Get file paths
+	FString BaseColorFile, NormalFile, RoughnessFile, MetallicFile, EmissiveFile, ORMFile;
+	GetStringParam(Params, TEXT("base_color_file"), BaseColorFile, false);
+	GetStringParam(Params, TEXT("normal_file"), NormalFile, false);
+	GetStringParam(Params, TEXT("roughness_file"), RoughnessFile, false);
+	GetStringParam(Params, TEXT("metallic_file"), MetallicFile, false);
+	GetStringParam(Params, TEXT("emissive_file"), EmissiveFile, false);
+	GetStringParam(Params, TEXT("orm_file"), ORMFile, false);
+	
+	bool bTwoSided = false;
+	GetBoolParam(Params, TEXT("two_sided"), bTwoSided, false);
+	
+	FString BlendModeStr = TEXT("Opaque");
+	GetStringParam(Params, TEXT("blend_mode"), BlendModeStr, false);
+	
+	// Import textures and collect their paths
+	TArray<TSharedPtr<FJsonValue>> ImportedTexturesArray;
+	FString BaseColorTexturePath, NormalTexturePath, RoughnessTexturePath, MetallicTexturePath, EmissiveTexturePath, ORMTexturePath;
+	FString Error;
+	
+	auto ImportAndTrack = [&](const FString& FilePath, const FString& Suffix, bool bIsNormal, FString& OutTexturePath) -> bool
+	{
+		if (FilePath.IsEmpty())
+		{
+			return true; // Not an error, just nothing to import
+		}
+		
+		// Configure import settings BEFORE import to ensure correct gamma space
+		FTextureImportSettings ImportSettings;
+		if (bIsNormal)
+		{
+			ImportSettings.CompressionSettings = TC_Normalmap;
+			ImportSettings.bSRGB = false;
+		}
+		else if (Suffix == TEXT("Roughness") || Suffix == TEXT("Metallic") || Suffix == TEXT("ORM"))
+		{
+			ImportSettings.CompressionSettings = TC_Masks;
+			ImportSettings.bSRGB = false;
+		}
+		else
+		{
+			ImportSettings.CompressionSettings = TC_Default;
+			ImportSettings.bSRGB = true;
+		}
+		
+		FString TextureName = MaterialName + TEXT("_") + Suffix;
+		UTexture2D* Texture = ImportTextureFromFile(FilePath, TexturePath, TextureName, Error, ImportSettings);
+		
+		if (!Texture)
+		{
+			return false;
+		}
+		
+		OutTexturePath = Texture->GetPathName();
+		
+		TSharedPtr<FJsonObject> TextureInfo = MakeShared<FJsonObject>();
+		TextureInfo->SetStringField(TEXT("path"), OutTexturePath);
+		TextureInfo->SetStringField(TEXT("type"), Suffix);
+		ImportedTexturesArray.Add(MakeShared<FJsonValueObject>(TextureInfo));
+		
+		return true;
+	};
+	
+	// Import all textures
+	if (!ImportAndTrack(BaseColorFile, TEXT("BaseColor"), false, BaseColorTexturePath))
+	{
+		return FECACommandResult::Error(FString::Printf(TEXT("Failed to import base color texture: %s"), *Error));
+	}
+	
+	if (!ImportAndTrack(NormalFile, TEXT("Normal"), true, NormalTexturePath))
+	{
+		return FECACommandResult::Error(FString::Printf(TEXT("Failed to import normal texture: %s"), *Error));
+	}
+	
+	if (!ImportAndTrack(RoughnessFile, TEXT("Roughness"), false, RoughnessTexturePath))
+	{
+		return FECACommandResult::Error(FString::Printf(TEXT("Failed to import roughness texture: %s"), *Error));
+	}
+	
+	if (!ImportAndTrack(MetallicFile, TEXT("Metallic"), false, MetallicTexturePath))
+	{
+		return FECACommandResult::Error(FString::Printf(TEXT("Failed to import metallic texture: %s"), *Error));
+	}
+	
+	if (!ImportAndTrack(EmissiveFile, TEXT("Emissive"), false, EmissiveTexturePath))
+	{
+		return FECACommandResult::Error(FString::Printf(TEXT("Failed to import emissive texture: %s"), *Error));
+	}
+	
+	if (!ImportAndTrack(ORMFile, TEXT("ORM"), false, ORMTexturePath))
+	{
+		return FECACommandResult::Error(FString::Printf(TEXT("Failed to import ORM texture: %s"), *Error));
+	}
+	
+	// Create the material package
+	FString PackagePath = MaterialPath + MaterialName;
+	UPackage* Package = CreatePackage(*PackagePath);
+	if (!Package)
+	{
+		return FECACommandResult::Error(TEXT("Failed to create material package"));
+	}
+	
+	// Create the material
+	UMaterialFactoryNew* MaterialFactory = NewObject<UMaterialFactoryNew>();
+	UMaterial* NewMaterial = Cast<UMaterial>(MaterialFactory->FactoryCreateNew(
+		UMaterial::StaticClass(),
+		Package,
+		*MaterialName,
+		RF_Public | RF_Standalone,
+		nullptr,
+		GWarn
+	));
+	
+	if (!NewMaterial)
+	{
+		return FECACommandResult::Error(TEXT("Failed to create material"));
+	}
+	
+	NewMaterial->TwoSided = bTwoSided;
+	NewMaterial->BlendMode = GetBlendModeFromString(BlendModeStr);
+	
+	int32 NodePosX = -400;
+	int32 NodePosY = 0;
+	const int32 NodeSpacingY = 200;
+	
+	// Helper to add texture sampler
+	auto AddTextureSampler = [&](const FString& TexturePathStr, const FString& ParamName, EMaterialProperty Property, int32 OutputIdx = 0) -> UMaterialExpressionTextureSampleParameter2D*
+	{
+		if (TexturePathStr.IsEmpty())
+		{
+			return nullptr;
+		}
+		
+		UTexture2D* Texture = LoadTextureByPath(TexturePathStr);
+		if (!Texture)
+		{
+			return nullptr;
+		}
+		
+		UMaterialExpressionTextureSampleParameter2D* Sampler = NewObject<UMaterialExpressionTextureSampleParameter2D>(NewMaterial);
+		Sampler->Texture = Texture;
+		Sampler->ParameterName = FName(*ParamName);
+		// Set appropriate sampler type based on material property and texture format
+		if (Property == MP_Normal)
+		{
+			Sampler->SamplerType = SAMPLERTYPE_Normal;
+		}
+		else if (Property == MP_Roughness || Property == MP_Metallic || Property == MP_AmbientOcclusion)
+		{
+			// Use LinearGrayscale for single-channel textures, Masks for packed RGB textures
+			if (Texture->CompressionSettings == TC_Grayscale || Texture->CompressionSettings == TC_Alpha)
+			{
+				Sampler->SamplerType = SAMPLERTYPE_LinearGrayscale;
+			}
+			else
+			{
+				Sampler->SamplerType = SAMPLERTYPE_Masks;
+			}
+		}
+		else
+		{
+			Sampler->SamplerType = SAMPLERTYPE_Color;
+		}
+		Sampler->MaterialExpressionEditorX = NodePosX;
+		Sampler->MaterialExpressionEditorY = NodePosY;
+		NodePosY += NodeSpacingY;
+		
+		NewMaterial->GetEditorOnlyData()->ExpressionCollection.Expressions.Add(Sampler);
+		
+		FExpressionInput* Input = NewMaterial->GetExpressionInputForProperty(Property);
+		if (Input)
+		{
+			Input->Expression = Sampler;
+			Input->OutputIndex = OutputIdx;
+		}
+		
+		return Sampler;
+	};
+	
+	// Add all texture inputs
+	AddTextureSampler(BaseColorTexturePath, TEXT("BaseColorTexture"), MP_BaseColor);
+	AddTextureSampler(NormalTexturePath, TEXT("NormalTexture"), MP_Normal);
+	AddTextureSampler(EmissiveTexturePath, TEXT("EmissiveTexture"), MP_EmissiveColor);
+	
+	// Handle ORM or individual roughness/metallic
+	if (!ORMTexturePath.IsEmpty())
+	{
+		// ORM texture: R = Occlusion (AO), G = Roughness, B = Metallic
+		UTexture2D* ORMTexture = LoadTextureByPath(ORMTexturePath);
+		if (ORMTexture)
+		{
+			UMaterialExpressionTextureSampleParameter2D* ORMSampler = NewObject<UMaterialExpressionTextureSampleParameter2D>(NewMaterial);
+			ORMSampler->Texture = ORMTexture;
+			ORMSampler->ParameterName = TEXT("ORMTexture");
+			ORMSampler->SamplerType = SAMPLERTYPE_Masks;
+			ORMSampler->MaterialExpressionEditorX = NodePosX;
+			ORMSampler->MaterialExpressionEditorY = NodePosY;
+			NodePosY += NodeSpacingY;
+			
+			NewMaterial->GetEditorOnlyData()->ExpressionCollection.Expressions.Add(ORMSampler);
+			
+			// Connect G channel to Roughness
+			NewMaterial->GetEditorOnlyData()->Roughness.Expression = ORMSampler;
+			NewMaterial->GetEditorOnlyData()->Roughness.OutputIndex = 2; // G channel
+			
+			// Connect B channel to Metallic
+			NewMaterial->GetEditorOnlyData()->Metallic.Expression = ORMSampler;
+			NewMaterial->GetEditorOnlyData()->Metallic.OutputIndex = 3; // B channel
+			
+			// Connect R channel to AO
+			NewMaterial->GetEditorOnlyData()->AmbientOcclusion.Expression = ORMSampler;
+			NewMaterial->GetEditorOnlyData()->AmbientOcclusion.OutputIndex = 1; // R channel
+		}
+	}
+	else
+	{
+		AddTextureSampler(RoughnessTexturePath, TEXT("RoughnessTexture"), MP_Roughness);
+		AddTextureSampler(MetallicTexturePath, TEXT("MetallicTexture"), MP_Metallic);
+	}
+	
+	// Compile material
+	NewMaterial->PreEditChange(nullptr);
+	NewMaterial->PostEditChange();
+	
+	FAssetRegistryModule::AssetCreated(NewMaterial);
+	Package->MarkPackageDirty();
+	
+	TSharedPtr<FJsonObject> Result = MakeResult();
+	Result->SetStringField(TEXT("material_path"), NewMaterial->GetPathName());
+	Result->SetStringField(TEXT("material_name"), NewMaterial->GetName());
+	Result->SetArrayField(TEXT("imported_textures"), ImportedTexturesArray);
+	Result->SetNumberField(TEXT("texture_count"), ImportedTexturesArray.Num());
+	Result->SetBoolField(TEXT("two_sided"), bTwoSided);
+	Result->SetStringField(TEXT("blend_mode"), BlendModeStr);
+	
+	// Include any compilation errors
+	AddMaterialCompilationErrorsToResult(Result, NewMaterial);
+	
+	return FECACommandResult::Success(Result);
+}
+
+//------------------------------------------------------------------------------
+// GetTextureInfo
+//------------------------------------------------------------------------------
+
+FECACommandResult FECACommand_GetTextureInfo::Execute(const TSharedPtr<FJsonObject>& Params)
+{
+	FString TexturePath;
+	if (!GetStringParam(Params, TEXT("texture_path"), TexturePath))
+	{
+		return FECACommandResult::Error(TEXT("Missing required parameter: texture_path"));
+	}
+	
+	UTexture2D* Texture = LoadTextureByPath(TexturePath);
+	if (!Texture)
+	{
+		return FECACommandResult::Error(FString::Printf(TEXT("Texture not found: %s"), *TexturePath));
+	}
+	
+	TSharedPtr<FJsonObject> Result = MakeResult();
+	Result->SetStringField(TEXT("name"), Texture->GetName());
+	Result->SetStringField(TEXT("path"), Texture->GetPathName());
+	Result->SetNumberField(TEXT("width"), Texture->GetSizeX());
+	Result->SetNumberField(TEXT("height"), Texture->GetSizeY());
+	Result->SetBoolField(TEXT("srgb"), Texture->SRGB);
+	Result->SetStringField(TEXT("compression_settings"), UEnum::GetValueAsString(Texture->CompressionSettings));
+	Result->SetStringField(TEXT("lod_group"), UEnum::GetValueAsString(Texture->LODGroup));
+	Result->SetBoolField(TEXT("has_alpha"), Texture->HasAlphaChannel());
+	Result->SetNumberField(TEXT("num_mips"), Texture->GetNumMips());
+	
+	return FECACommandResult::Success(Result);
+}
+
+//------------------------------------------------------------------------------
+// GetMaterialInfo
+//------------------------------------------------------------------------------
+
+FECACommandResult FECACommand_GetMaterialInfo::Execute(const TSharedPtr<FJsonObject>& Params)
+{
+	FString MaterialPath;
+	if (!GetStringParam(Params, TEXT("material_path"), MaterialPath))
+	{
+		return FECACommandResult::Error(TEXT("Missing required parameter: material_path"));
+	}
+	
+	UMaterial* Material = LoadMaterialByPath(MaterialPath);
+	if (!Material)
+	{
+		return FECACommandResult::Error(FString::Printf(TEXT("Material not found: %s"), *MaterialPath));
+	}
+	
+	TSharedPtr<FJsonObject> Result = MakeResult();
+	Result->SetStringField(TEXT("name"), Material->GetName());
+	Result->SetStringField(TEXT("path"), Material->GetPathName());
+	Result->SetBoolField(TEXT("two_sided"), Material->TwoSided);
+	Result->SetStringField(TEXT("blend_mode"), UEnum::GetValueAsString(Material->BlendMode));
+	Result->SetStringField(TEXT("shading_model"), UEnum::GetValueAsString(Material->GetShadingModels().GetFirstShadingModel()));
+	
+	// Get texture parameters
+	TArray<TSharedPtr<FJsonValue>> TexturesArray;
+	const auto& Expressions = Material->GetEditorOnlyData()->ExpressionCollection.Expressions;
+	for (UMaterialExpression* Expression : Expressions)
+	{
+		if (UMaterialExpressionTextureSample* TextureSample = Cast<UMaterialExpressionTextureSample>(Expression))
+		{
+			if (TextureSample->Texture)
+			{
+				TSharedPtr<FJsonObject> TextureInfo = MakeShared<FJsonObject>();
+				TextureInfo->SetStringField(TEXT("texture_name"), TextureSample->Texture->GetName());
+				TextureInfo->SetStringField(TEXT("texture_path"), TextureSample->Texture->GetPathName());
+				TexturesArray.Add(MakeShared<FJsonValueObject>(TextureInfo));
+			}
+		}
+	}
+	Result->SetArrayField(TEXT("textures"), TexturesArray);
+	
+	// Get scalar parameters
+	TArray<TSharedPtr<FJsonValue>> ScalarsArray;
+	for (UMaterialExpression* Expression : Expressions)
+	{
+		if (UMaterialExpressionScalarParameter* ScalarParam = Cast<UMaterialExpressionScalarParameter>(Expression))
+		{
+			TSharedPtr<FJsonObject> ParamInfo = MakeShared<FJsonObject>();
+			ParamInfo->SetStringField(TEXT("name"), ScalarParam->ParameterName.ToString());
+			ParamInfo->SetNumberField(TEXT("value"), ScalarParam->DefaultValue);
+			ScalarsArray.Add(MakeShared<FJsonValueObject>(ParamInfo));
+		}
+	}
+	Result->SetArrayField(TEXT("scalar_parameters"), ScalarsArray);
+	
+	// Get vector parameters
+	TArray<TSharedPtr<FJsonValue>> VectorsArray;
+	for (UMaterialExpression* Expression : Expressions)
+	{
+		if (UMaterialExpressionVectorParameter* VectorParam = Cast<UMaterialExpressionVectorParameter>(Expression))
+		{
+			TSharedPtr<FJsonObject> ParamInfo = MakeShared<FJsonObject>();
+			ParamInfo->SetStringField(TEXT("name"), VectorParam->ParameterName.ToString());
+			TSharedPtr<FJsonObject> ValueObj = MakeShared<FJsonObject>();
+			ValueObj->SetNumberField(TEXT("r"), VectorParam->DefaultValue.R);
+			ValueObj->SetNumberField(TEXT("g"), VectorParam->DefaultValue.G);
+			ValueObj->SetNumberField(TEXT("b"), VectorParam->DefaultValue.B);
+			ValueObj->SetNumberField(TEXT("a"), VectorParam->DefaultValue.A);
+			ParamInfo->SetObjectField(TEXT("value"), ValueObj);
+			VectorsArray.Add(MakeShared<FJsonValueObject>(ParamInfo));
+		}
+	}
+	Result->SetArrayField(TEXT("vector_parameters"), VectorsArray);
+	
+	return FECACommandResult::Success(Result);
+}
+
+//------------------------------------------------------------------------------
+// ListTextures
+//------------------------------------------------------------------------------
+
+FECACommandResult FECACommand_ListTextures::Execute(const TSharedPtr<FJsonObject>& Params)
+{
+	FString Path = TEXT("/Game/");
+	bool bRecursive = true;
+	GetStringParam(Params, TEXT("path"), Path, false);
+	GetBoolParam(Params, TEXT("recursive"), bRecursive, false);
+	
+	FAssetRegistryModule& AssetRegistryModule = FModuleManager::LoadModuleChecked<FAssetRegistryModule>("AssetRegistry");
+	IAssetRegistry& AssetRegistry = AssetRegistryModule.Get();
+	
+	TArray<FAssetData> AssetList;
+	AssetRegistry.GetAssetsByPath(*Path, AssetList, bRecursive);
+	
+	TArray<TSharedPtr<FJsonValue>> TexturesArray;
+	
+	for (const FAssetData& Asset : AssetList)
+	{
+		if (Asset.AssetClassPath == UTexture2D::StaticClass()->GetClassPathName())
+		{
+			TSharedPtr<FJsonObject> TextureObj = MakeShared<FJsonObject>();
+			TextureObj->SetStringField(TEXT("name"), Asset.AssetName.ToString());
+			TextureObj->SetStringField(TEXT("path"), Asset.GetObjectPathString());
+			TexturesArray.Add(MakeShared<FJsonValueObject>(TextureObj));
+		}
+	}
+	
+	TSharedPtr<FJsonObject> Result = MakeResult();
+	Result->SetArrayField(TEXT("textures"), TexturesArray);
+	Result->SetNumberField(TEXT("count"), TexturesArray.Num());
+	
+	return FECACommandResult::Success(Result);
+}
+
+//------------------------------------------------------------------------------
+// ListMaterials
+//------------------------------------------------------------------------------
+
+FECACommandResult FECACommand_ListMaterials::Execute(const TSharedPtr<FJsonObject>& Params)
+{
+	FString Path = TEXT("/Game/");
+	bool bRecursive = true;
+	GetStringParam(Params, TEXT("path"), Path, false);
+	GetBoolParam(Params, TEXT("recursive"), bRecursive, false);
+	
+	FAssetRegistryModule& AssetRegistryModule = FModuleManager::LoadModuleChecked<FAssetRegistryModule>("AssetRegistry");
+	IAssetRegistry& AssetRegistry = AssetRegistryModule.Get();
+	
+	TArray<FAssetData> AssetList;
+	AssetRegistry.GetAssetsByPath(*Path, AssetList, bRecursive);
+	
+	TArray<TSharedPtr<FJsonValue>> MaterialsArray;
+	
+	for (const FAssetData& Asset : AssetList)
+	{
+		if (Asset.AssetClassPath == UMaterial::StaticClass()->GetClassPathName() ||
+			Asset.AssetClassPath == UMaterialInstance::StaticClass()->GetClassPathName() ||
+			Asset.AssetClassPath == UMaterialInstanceConstant::StaticClass()->GetClassPathName())
+		{
+			TSharedPtr<FJsonObject> MaterialObj = MakeShared<FJsonObject>();
+			MaterialObj->SetStringField(TEXT("name"), Asset.AssetName.ToString());
+			MaterialObj->SetStringField(TEXT("path"), Asset.GetObjectPathString());
+			MaterialObj->SetStringField(TEXT("type"), Asset.AssetClassPath.GetAssetName().ToString());
+			MaterialsArray.Add(MakeShared<FJsonValueObject>(MaterialObj));
+		}
+	}
+	
+	TSharedPtr<FJsonObject> Result = MakeResult();
+	Result->SetArrayField(TEXT("materials"), MaterialsArray);
+	Result->SetNumberField(TEXT("count"), MaterialsArray.Num());
+	
+	return FECACommandResult::Success(Result);
+}
+
+//------------------------------------------------------------------------------
+// SetMaterialTextureParam
+//------------------------------------------------------------------------------
+
+FECACommandResult FECACommand_SetMaterialTextureParam::Execute(const TSharedPtr<FJsonObject>& Params)
+{
+	FString MaterialPath;
+	if (!GetStringParam(Params, TEXT("material_path"), MaterialPath))
+	{
+		return FECACommandResult::Error(TEXT("Missing required parameter: material_path"));
+	}
+	
+	FString ParameterName;
+	if (!GetStringParam(Params, TEXT("parameter_name"), ParameterName))
+	{
+		return FECACommandResult::Error(TEXT("Missing required parameter: parameter_name"));
+	}
+	
+	FString TexturePath;
+	if (!GetStringParam(Params, TEXT("texture_path"), TexturePath))
+	{
+		return FECACommandResult::Error(TEXT("Missing required parameter: texture_path"));
+	}
+	
+	UMaterial* Material = LoadMaterialByPath(MaterialPath);
+	if (!Material)
+	{
+		return FECACommandResult::Error(FString::Printf(TEXT("Material not found: %s"), *MaterialPath));
+	}
+	
+	UTexture2D* Texture = LoadTextureByPath(TexturePath);
+	if (!Texture)
+	{
+		return FECACommandResult::Error(FString::Printf(TEXT("Texture not found: %s"), *TexturePath));
+	}
+	
+	// Find the texture parameter
+	bool bFound = false;
+	for (UMaterialExpression* Expression : Material->GetEditorOnlyData()->ExpressionCollection.Expressions)
+	{
+		if (UMaterialExpressionTextureSampleParameter2D* TextureParam = Cast<UMaterialExpressionTextureSampleParameter2D>(Expression))
+		{
+			if (TextureParam->ParameterName.ToString() == ParameterName)
+			{
+				TextureParam->Texture = Texture;
+				bFound = true;
+				break;
+			}
+		}
+	}
+	
+	if (!bFound)
+	{
+		return FECACommandResult::Error(FString::Printf(TEXT("Texture parameter not found: %s"), *ParameterName));
+	}
+	
+	Material->PreEditChange(nullptr);
+	Material->PostEditChange();
+	Material->MarkPackageDirty();
+	
+	TSharedPtr<FJsonObject> Result = MakeResult();
+	Result->SetStringField(TEXT("material_path"), MaterialPath);
+	Result->SetStringField(TEXT("parameter_name"), ParameterName);
+	Result->SetStringField(TEXT("texture_path"), TexturePath);
+	
+	return FECACommandResult::Success(Result);
+}
+
+//------------------------------------------------------------------------------
+// SetActorMaterial
+//------------------------------------------------------------------------------
+
+FECACommandResult FECACommand_SetActorMaterial::Execute(const TSharedPtr<FJsonObject>& Params)
+{
+	FString ActorName;
+	if (!GetStringParam(Params, TEXT("actor_name"), ActorName))
+	{
+		return FECACommandResult::Error(TEXT("Missing required parameter: actor_name"));
+	}
+	
+	FString MaterialPath;
+	if (!GetStringParam(Params, TEXT("material_path"), MaterialPath))
+	{
+		return FECACommandResult::Error(TEXT("Missing required parameter: material_path"));
+	}
+	
+	int32 SlotIndex = 0;
+	GetIntParam(Params, TEXT("slot_index"), SlotIndex, false);
+	
+	FString ComponentName;
+	GetStringParam(Params, TEXT("component_name"), ComponentName, false);
+	
+	// Find the actor
+	AActor* FoundActor = FindActorByName(ActorName);
+	if (!FoundActor)
+	{
+		return FECACommandResult::Error(FString::Printf(TEXT("Actor not found: %s"), *ActorName));
+	}
+	
+	// Load the material
+	UMaterialInterface* Material = LoadObject<UMaterialInterface>(nullptr, *MaterialPath);
+	if (!Material)
+	{
+		return FECACommandResult::Error(FString::Printf(TEXT("Material not found: %s"), *MaterialPath));
+	}
+	
+	// Find the mesh component
+	UPrimitiveComponent* TargetComponent = nullptr;
+	
+	if (!ComponentName.IsEmpty())
+	{
+		// Find specific component by name
+		TArray<UActorComponent*> Components;
+		FoundActor->GetComponents(Components);
+		for (UActorComponent* Component : Components)
+		{
+			if (Component->GetName() == ComponentName)
+			{
+				TargetComponent = Cast<UPrimitiveComponent>(Component);
+				break;
+			}
+		}
+		
+		if (!TargetComponent)
+		{
+			return FECACommandResult::Error(FString::Printf(TEXT("Component not found: %s"), *ComponentName));
+		}
+	}
+	else
+	{
+		// Find first mesh component
+		TargetComponent = FoundActor->FindComponentByClass<UStaticMeshComponent>();
+		if (!TargetComponent)
+		{
+			TargetComponent = FoundActor->FindComponentByClass<USkeletalMeshComponent>();
+		}
+		if (!TargetComponent)
+		{
+			TargetComponent = FoundActor->FindComponentByClass<UPrimitiveComponent>();
+		}
+		
+		if (!TargetComponent)
+		{
+			return FECACommandResult::Error(TEXT("No mesh component found on actor"));
+		}
+	}
+	
+	// Validate slot index
+	int32 NumMaterials = TargetComponent->GetNumMaterials();
+	if (SlotIndex < 0 || SlotIndex >= NumMaterials)
+	{
+		return FECACommandResult::Error(FString::Printf(TEXT("Invalid slot index %d. Actor has %d material slots (0-%d)"), SlotIndex, NumMaterials, NumMaterials - 1));
+	}
+	
+	// Set the material
+	TargetComponent->SetMaterial(SlotIndex, Material);
+	
+	TSharedPtr<FJsonObject> Result = MakeResult();
+	Result->SetStringField(TEXT("actor_name"), ActorName);
+	Result->SetStringField(TEXT("component_name"), TargetComponent->GetName());
+	Result->SetStringField(TEXT("material_path"), MaterialPath);
+	Result->SetNumberField(TEXT("slot_index"), SlotIndex);
+	
+	return FECACommandResult::Success(Result);
+}
+
+//------------------------------------------------------------------------------
+// GetActorMaterials
+//------------------------------------------------------------------------------
+
+FECACommandResult FECACommand_GetActorMaterials::Execute(const TSharedPtr<FJsonObject>& Params)
+{
+	FString ActorName;
+	if (!GetStringParam(Params, TEXT("actor_name"), ActorName))
+	{
+		return FECACommandResult::Error(TEXT("Missing required parameter: actor_name"));
+	}
+	
+	// Find the actor
+	AActor* FoundActor = FindActorByName(ActorName);
+	if (!FoundActor)
+	{
+		return FECACommandResult::Error(FString::Printf(TEXT("Actor not found: %s"), *ActorName));
+	}
+	
+	TArray<TSharedPtr<FJsonValue>> ComponentsArray;
+	
+	// Get all primitive components
+	TArray<UPrimitiveComponent*> PrimitiveComponents;
+	FoundActor->GetComponents<UPrimitiveComponent>(PrimitiveComponents);
+	
+	for (UPrimitiveComponent* Component : PrimitiveComponents)
+	{
+		TSharedPtr<FJsonObject> CompObj = MakeShared<FJsonObject>();
+		CompObj->SetStringField(TEXT("component_name"), Component->GetName());
+		CompObj->SetStringField(TEXT("component_class"), Component->GetClass()->GetName());
+		
+		TArray<TSharedPtr<FJsonValue>> MaterialsArray;
+		int32 NumMaterials = Component->GetNumMaterials();
+		
+		for (int32 i = 0; i < NumMaterials; i++)
+		{
+			UMaterialInterface* Material = Component->GetMaterial(i);
+			TSharedPtr<FJsonObject> MatObj = MakeShared<FJsonObject>();
+			MatObj->SetNumberField(TEXT("slot_index"), i);
+			if (Material)
+			{
+				MatObj->SetStringField(TEXT("material_name"), Material->GetName());
+				MatObj->SetStringField(TEXT("material_path"), Material->GetPathName());
+			}
+			else
+			{
+				MatObj->SetStringField(TEXT("material_name"), TEXT("None"));
+				MatObj->SetStringField(TEXT("material_path"), TEXT(""));
+			}
+			MaterialsArray.Add(MakeShared<FJsonValueObject>(MatObj));
+		}
+		
+		CompObj->SetArrayField(TEXT("materials"), MaterialsArray);
+		CompObj->SetNumberField(TEXT("material_count"), NumMaterials);
+		ComponentsArray.Add(MakeShared<FJsonValueObject>(CompObj));
+	}
+	
+	TSharedPtr<FJsonObject> Result = MakeResult();
+	Result->SetStringField(TEXT("actor_name"), ActorName);
+	Result->SetArrayField(TEXT("components"), ComponentsArray);
+	Result->SetNumberField(TEXT("component_count"), ComponentsArray.Num());
+	
+	return FECACommandResult::Success(Result);
+}
+
+//------------------------------------------------------------------------------
+// SetStaticMeshMaterial
+//------------------------------------------------------------------------------
+
+FECACommandResult FECACommand_SetStaticMeshMaterial::Execute(const TSharedPtr<FJsonObject>& Params)
+{
+	FString MeshPath;
+	if (!GetStringParam(Params, TEXT("mesh_path"), MeshPath))
+	{
+		return FECACommandResult::Error(TEXT("Missing required parameter: mesh_path"));
+	}
+	
+	FString MaterialPath;
+	if (!GetStringParam(Params, TEXT("material_path"), MaterialPath))
+	{
+		return FECACommandResult::Error(TEXT("Missing required parameter: material_path"));
+	}
+	
+	int32 SlotIndex = 0;
+	GetIntParam(Params, TEXT("slot_index"), SlotIndex, false);
+	
+	// Load the static mesh
+	UStaticMesh* StaticMesh = LoadObject<UStaticMesh>(nullptr, *MeshPath);
+	if (!StaticMesh)
+	{
+		return FECACommandResult::Error(FString::Printf(TEXT("Static mesh not found: %s"), *MeshPath));
+	}
+	
+	// Load the material
+	UMaterialInterface* Material = LoadObject<UMaterialInterface>(nullptr, *MaterialPath);
+	if (!Material)
+	{
+		return FECACommandResult::Error(FString::Printf(TEXT("Material not found: %s"), *MaterialPath));
+	}
+	
+	// Validate slot index
+	int32 NumMaterials = StaticMesh->GetStaticMaterials().Num();
+	if (SlotIndex < 0 || SlotIndex >= NumMaterials)
+	{
+		return FECACommandResult::Error(FString::Printf(TEXT("Invalid slot index %d. Mesh has %d material slots (0-%d)"), SlotIndex, NumMaterials, NumMaterials - 1));
+	}
+	
+	// Set the material
+	StaticMesh->SetMaterial(SlotIndex, Material);
+	StaticMesh->MarkPackageDirty();
+	
+	TSharedPtr<FJsonObject> Result = MakeResult();
+	Result->SetStringField(TEXT("mesh_path"), MeshPath);
+	Result->SetStringField(TEXT("material_path"), MaterialPath);
+	Result->SetNumberField(TEXT("slot_index"), SlotIndex);
+	
+	return FECACommandResult::Success(Result);
+}
+
+//------------------------------------------------------------------------------
+// CreateMaterialInstance
+//------------------------------------------------------------------------------
+
+FECACommandResult FECACommand_CreateMaterialInstance::Execute(const TSharedPtr<FJsonObject>& Params)
+{
+	FString InstanceName;
+	if (!GetStringParam(Params, TEXT("instance_name"), InstanceName))
+	{
+		return FECACommandResult::Error(TEXT("Missing required parameter: instance_name"));
+	}
+	
+	FString ParentMaterialPath;
+	if (!GetStringParam(Params, TEXT("parent_material"), ParentMaterialPath))
+	{
+		return FECACommandResult::Error(TEXT("Missing required parameter: parent_material"));
+	}
+	
+	FString Path = TEXT("/Game/Materials/");
+	GetStringParam(Params, TEXT("path"), Path, false);
+	
+	// Ensure path ends with /
+	if (!Path.EndsWith(TEXT("/")))
+	{
+		Path += TEXT("/");
+	}
+	
+	// Load parent material
+	UMaterialInterface* ParentMaterial = LoadObject<UMaterialInterface>(nullptr, *ParentMaterialPath);
+	if (!ParentMaterial)
+	{
+		return FECACommandResult::Error(FString::Printf(TEXT("Parent material not found: %s"), *ParentMaterialPath));
+	}
+	
+	// Create package path
+	FString PackagePath = Path + InstanceName;
+	
+	// Check if asset already exists at this path
+	if (UObject* ExistingAsset = StaticLoadObject(UObject::StaticClass(), nullptr, *PackagePath))
+	{
+		if (UMaterialInstanceConstant* ExistingMIC = Cast<UMaterialInstanceConstant>(ExistingAsset))
+		{
+			TSharedPtr<FJsonObject> Result = MakeResult();
+			Result->SetStringField(TEXT("instance_path"), PackagePath);
+			Result->SetStringField(TEXT("instance_name"), InstanceName);
+			Result->SetBoolField(TEXT("already_exists"), true);
+			Result->SetStringField(TEXT("message"), TEXT("Material Instance already exists at this path"));
+			return FECACommandResult::Success(Result);
+		}
+		else
+		{
+			return FECACommandResult::Error(FString::Printf(TEXT("Asset already exists at path '%s' but is not a Material Instance (it's a %s)"), *PackagePath, *ExistingAsset->GetClass()->GetName()));
+		}
+	}
+	
+	// Create package
+	UPackage* Package = CreatePackage(*PackagePath);
+	if (!Package)
+	{
+		return FECACommandResult::Error(TEXT("Failed to create package"));
+	}
+	
+	// Create material instance factory
+	UMaterialInstanceConstantFactoryNew* Factory = NewObject<UMaterialInstanceConstantFactoryNew>();
+	Factory->InitialParent = ParentMaterial;
+	
+	// Create the material instance
+	UMaterialInstanceConstant* MaterialInstance = Cast<UMaterialInstanceConstant>(Factory->FactoryCreateNew(
+		UMaterialInstanceConstant::StaticClass(),
+		Package,
+		*InstanceName,
+		RF_Public | RF_Standalone,
+		nullptr,
+		GWarn
+	));
+	
+	if (!MaterialInstance)
+	{
+		return FECACommandResult::Error(TEXT("Failed to create material instance"));
+	}
+	
+	// Notify asset registry
+	FAssetRegistryModule::AssetCreated(MaterialInstance);
+	Package->MarkPackageDirty();
+	
+	TSharedPtr<FJsonObject> Result = MakeResult();
+	Result->SetStringField(TEXT("instance_path"), MaterialInstance->GetPathName());
+	Result->SetStringField(TEXT("instance_name"), MaterialInstance->GetName());
+	Result->SetStringField(TEXT("parent_material"), ParentMaterialPath);
+	
+	return FECACommandResult::Success(Result);
+}
+
+//------------------------------------------------------------------------------
+// SetMaterialInstanceScalarParam
+//------------------------------------------------------------------------------
+
+FECACommandResult FECACommand_SetMaterialInstanceScalarParam::Execute(const TSharedPtr<FJsonObject>& Params)
+{
+	FString InstancePath;
+	if (!GetStringParam(Params, TEXT("instance_path"), InstancePath))
+	{
+		return FECACommandResult::Error(TEXT("Missing required parameter: instance_path"));
+	}
+	
+	FString ParameterName;
+	if (!GetStringParam(Params, TEXT("parameter_name"), ParameterName))
+	{
+		return FECACommandResult::Error(TEXT("Missing required parameter: parameter_name"));
+	}
+	
+	double Value = 0.0;
+	if (!GetFloatParam(Params, TEXT("value"), Value))
+	{
+		return FECACommandResult::Error(TEXT("Missing required parameter: value"));
+	}
+	
+	// Load material instance
+	UMaterialInstanceConstant* MaterialInstance = LoadObject<UMaterialInstanceConstant>(nullptr, *InstancePath);
+	if (!MaterialInstance)
+	{
+		return FECACommandResult::Error(FString::Printf(TEXT("Material instance not found: %s"), *InstancePath));
+	}
+	
+	// Set the parameter
+	MaterialInstance->SetScalarParameterValueEditorOnly(FName(*ParameterName), Value);
+	MaterialInstance->MarkPackageDirty();
+	
+	TSharedPtr<FJsonObject> Result = MakeResult();
+	Result->SetStringField(TEXT("instance_path"), InstancePath);
+	Result->SetStringField(TEXT("parameter_name"), ParameterName);
+	Result->SetNumberField(TEXT("value"), Value);
+	
+	return FECACommandResult::Success(Result);
+}
+
+//------------------------------------------------------------------------------
+// SetMaterialInstanceVectorParam
+//------------------------------------------------------------------------------
+
+FECACommandResult FECACommand_SetMaterialInstanceVectorParam::Execute(const TSharedPtr<FJsonObject>& Params)
+{
+	FString InstancePath;
+	if (!GetStringParam(Params, TEXT("instance_path"), InstancePath))
+	{
+		return FECACommandResult::Error(TEXT("Missing required parameter: instance_path"));
+	}
+	
+	FString ParameterName;
+	if (!GetStringParam(Params, TEXT("parameter_name"), ParameterName))
+	{
+		return FECACommandResult::Error(TEXT("Missing required parameter: parameter_name"));
+	}
+	
+	const TSharedPtr<FJsonObject>* ValueObj;
+	if (!GetObjectParam(Params, TEXT("value"), ValueObj))
+	{
+		return FECACommandResult::Error(TEXT("Missing required parameter: value"));
+	}
+	
+	// Load material instance
+	UMaterialInstanceConstant* MaterialInstance = LoadObject<UMaterialInstanceConstant>(nullptr, *InstancePath);
+	if (!MaterialInstance)
+	{
+		return FECACommandResult::Error(FString::Printf(TEXT("Material instance not found: %s"), *InstancePath));
+	}
+	
+	// Parse the color/vector value
+	FLinearColor Color;
+	if ((*ValueObj)->HasField(TEXT("r")))
+	{
+		Color.R = (*ValueObj)->GetNumberField(TEXT("r"));
+		Color.G = (*ValueObj)->GetNumberField(TEXT("g"));
+		Color.B = (*ValueObj)->GetNumberField(TEXT("b"));
+		Color.A = (*ValueObj)->HasField(TEXT("a")) ? (*ValueObj)->GetNumberField(TEXT("a")) : 1.0f;
+	}
+	else if ((*ValueObj)->HasField(TEXT("x")))
+	{
+		Color.R = (*ValueObj)->GetNumberField(TEXT("x"));
+		Color.G = (*ValueObj)->GetNumberField(TEXT("y"));
+		Color.B = (*ValueObj)->GetNumberField(TEXT("z"));
+		Color.A = (*ValueObj)->HasField(TEXT("w")) ? (*ValueObj)->GetNumberField(TEXT("w")) : 1.0f;
+	}
+	else
+	{
+		return FECACommandResult::Error(TEXT("Value must have r,g,b or x,y,z fields"));
+	}
+	
+	// Set the parameter
+	MaterialInstance->SetVectorParameterValueEditorOnly(FName(*ParameterName), Color);
+	MaterialInstance->MarkPackageDirty();
+	
+	TSharedPtr<FJsonObject> Result = MakeResult();
+	Result->SetStringField(TEXT("instance_path"), InstancePath);
+	Result->SetStringField(TEXT("parameter_name"), ParameterName);
+	
+	TSharedPtr<FJsonObject> ValueResult = MakeShared<FJsonObject>();
+	ValueResult->SetNumberField(TEXT("r"), Color.R);
+	ValueResult->SetNumberField(TEXT("g"), Color.G);
+	ValueResult->SetNumberField(TEXT("b"), Color.B);
+	ValueResult->SetNumberField(TEXT("a"), Color.A);
+	Result->SetObjectField(TEXT("value"), ValueResult);
+	
+	return FECACommandResult::Success(Result);
+}
+
+//------------------------------------------------------------------------------
+// SetMaterialInstanceTextureParam
+//------------------------------------------------------------------------------
+
+FECACommandResult FECACommand_SetMaterialInstanceTextureParam::Execute(const TSharedPtr<FJsonObject>& Params)
+{
+	FString InstancePath;
+	if (!GetStringParam(Params, TEXT("instance_path"), InstancePath))
+	{
+		return FECACommandResult::Error(TEXT("Missing required parameter: instance_path"));
+	}
+	
+	FString ParameterName;
+	if (!GetStringParam(Params, TEXT("parameter_name"), ParameterName))
+	{
+		return FECACommandResult::Error(TEXT("Missing required parameter: parameter_name"));
+	}
+	
+	FString TexturePath;
+	if (!GetStringParam(Params, TEXT("texture_path"), TexturePath))
+	{
+		return FECACommandResult::Error(TEXT("Missing required parameter: texture_path"));
+	}
+	
+	// Load material instance
+	UMaterialInstanceConstant* MaterialInstance = LoadObject<UMaterialInstanceConstant>(nullptr, *InstancePath);
+	if (!MaterialInstance)
+	{
+		return FECACommandResult::Error(FString::Printf(TEXT("Material instance not found: %s"), *InstancePath));
+	}
+	
+	// Load texture
+	UTexture* Texture = LoadObject<UTexture>(nullptr, *TexturePath);
+	if (!Texture)
+	{
+		return FECACommandResult::Error(FString::Printf(TEXT("Texture not found: %s"), *TexturePath));
+	}
+	
+	// Set the parameter
+	MaterialInstance->SetTextureParameterValueEditorOnly(FName(*ParameterName), Texture);
+	MaterialInstance->MarkPackageDirty();
+	
+	TSharedPtr<FJsonObject> Result = MakeResult();
+	Result->SetStringField(TEXT("instance_path"), InstancePath);
+	Result->SetStringField(TEXT("parameter_name"), ParameterName);
+	Result->SetStringField(TEXT("texture_path"), TexturePath);
+	
+	return FECACommandResult::Success(Result);
+}
+
+//------------------------------------------------------------------------------
+// ImportOBJ - OBJ File Parser and Static Mesh Creator
+//------------------------------------------------------------------------------
+
+// OBJ file parsing structures
+struct FOBJVertex
+{
+	FVector3f Position;
+};
+
+struct FOBJTexCoord
+{
+	FVector2f UV;
+};
+
+struct FOBJNormal
+{
+	FVector3f Normal;
+};
+
+struct FOBJFaceVertex
+{
+	int32 PositionIndex;  // 1-based index into vertex positions
+	int32 TexCoordIndex;  // 1-based index into texture coordinates (0 = none)
+	int32 NormalIndex;    // 1-based index into normals (0 = none)
+};
+
+struct FOBJFace
+{
+	TArray<FOBJFaceVertex> Vertices;
+};
+
+struct FOBJMaterialGroup
+{
+	FString MaterialName;
+	TArray<FOBJFace> Faces;
+};
+
+struct FOBJObject
+{
+	FString Name;
+	TArray<FOBJMaterialGroup> MaterialGroups;
+};
+
+struct FOBJMaterial
+{
+	FString Name;
+	FLinearColor DiffuseColor = FLinearColor(0.8f, 0.8f, 0.8f, 1.0f);
+	FLinearColor AmbientColor = FLinearColor(0.2f, 0.2f, 0.2f, 1.0f);
+	FLinearColor SpecularColor = FLinearColor(1.0f, 1.0f, 1.0f, 1.0f);
+	float Shininess = 32.0f;
+	float Opacity = 1.0f;
+	FString DiffuseTexture;
+	FString NormalTexture;
+	FString SpecularTexture;
+};
+
+struct FOBJParseResult
+{
+	TArray<FOBJVertex> Vertices;
+	TArray<FOBJTexCoord> TexCoords;
+	TArray<FOBJNormal> Normals;
+	TArray<FOBJObject> Objects;
+	TArray<FOBJMaterial> Materials;
+	FString MTLLibrary;
+	FString ErrorMessage;
+	bool bSuccess = false;
+};
+
+// Parse a single face vertex (format: v, v/vt, v/vt/vn, or v//vn)
+static FOBJFaceVertex ParseFaceVertex(const FString& Token)
+{
+	FOBJFaceVertex Result;
+	Result.PositionIndex = 0;
+	Result.TexCoordIndex = 0;
+	Result.NormalIndex = 0;
+	
+	TArray<FString> Parts;
+	Token.ParseIntoArray(Parts, TEXT("/"), false);
+	
+	if (Parts.Num() >= 1 && !Parts[0].IsEmpty())
+	{
+		Result.PositionIndex = FCString::Atoi(*Parts[0]);
+	}
+	if (Parts.Num() >= 2 && !Parts[1].IsEmpty())
+	{
+		Result.TexCoordIndex = FCString::Atoi(*Parts[1]);
+	}
+	if (Parts.Num() >= 3 && !Parts[2].IsEmpty())
+	{
+		Result.NormalIndex = FCString::Atoi(*Parts[2]);
+	}
+	
+	return Result;
+}
+
+// Parse MTL file for materials
+static TArray<FOBJMaterial> ParseMTLFile(const FString& MTLPath)
+{
+	TArray<FOBJMaterial> Materials;
+	
+	FString MTLContent;
+	if (!FFileHelper::LoadFileToString(MTLContent, *MTLPath))
+	{
+		return Materials;
+	}
+	
+	TArray<FString> Lines;
+	MTLContent.ParseIntoArrayLines(Lines);
+	
+	FOBJMaterial* CurrentMaterial = nullptr;
+	FString MTLDirectory = FPaths::GetPath(MTLPath);
+	
+	for (const FString& Line : Lines)
+	{
+		FString TrimmedLine = Line.TrimStartAndEnd();
+		if (TrimmedLine.IsEmpty() || TrimmedLine.StartsWith(TEXT("#")))
+		{
+			continue;
+		}
+		
+		TArray<FString> Tokens;
+		TrimmedLine.ParseIntoArrayWS(Tokens);
+		
+		if (Tokens.Num() == 0)
+		{
+			continue;
+		}
+		
+		FString Command = Tokens[0].ToLower();
+		
+		if (Command == TEXT("newmtl") && Tokens.Num() >= 2)
+		{
+			Materials.Add(FOBJMaterial());
+			CurrentMaterial = &Materials.Last();
+			CurrentMaterial->Name = Tokens[1];
+		}
+		else if (CurrentMaterial)
+		{
+			if (Command == TEXT("kd") && Tokens.Num() >= 4)
+			{
+				CurrentMaterial->DiffuseColor.R = FCString::Atof(*Tokens[1]);
+				CurrentMaterial->DiffuseColor.G = FCString::Atof(*Tokens[2]);
+				CurrentMaterial->DiffuseColor.B = FCString::Atof(*Tokens[3]);
+			}
+			else if (Command == TEXT("ka") && Tokens.Num() >= 4)
+			{
+				CurrentMaterial->AmbientColor.R = FCString::Atof(*Tokens[1]);
+				CurrentMaterial->AmbientColor.G = FCString::Atof(*Tokens[2]);
+				CurrentMaterial->AmbientColor.B = FCString::Atof(*Tokens[3]);
+			}
+			else if (Command == TEXT("ks") && Tokens.Num() >= 4)
+			{
+				CurrentMaterial->SpecularColor.R = FCString::Atof(*Tokens[1]);
+				CurrentMaterial->SpecularColor.G = FCString::Atof(*Tokens[2]);
+				CurrentMaterial->SpecularColor.B = FCString::Atof(*Tokens[3]);
+			}
+			else if (Command == TEXT("ns") && Tokens.Num() >= 2)
+			{
+				CurrentMaterial->Shininess = FCString::Atof(*Tokens[1]);
+			}
+			else if (Command == TEXT("d") && Tokens.Num() >= 2)
+			{
+				CurrentMaterial->Opacity = FCString::Atof(*Tokens[1]);
+			}
+			else if (Command == TEXT("tr") && Tokens.Num() >= 2)
+			{
+				CurrentMaterial->Opacity = 1.0f - FCString::Atof(*Tokens[1]);
+			}
+			else if (Command == TEXT("map_kd") && Tokens.Num() >= 2)
+			{
+				FString TexturePath = Tokens[1];
+				if (!FPaths::IsRelative(TexturePath))
+				{
+					CurrentMaterial->DiffuseTexture = TexturePath;
+				}
+				else
+				{
+					CurrentMaterial->DiffuseTexture = FPaths::Combine(MTLDirectory, TexturePath);
+				}
+			}
+			else if ((Command == TEXT("map_bump") || Command == TEXT("bump")) && Tokens.Num() >= 2)
+			{
+				FString TexturePath = Tokens[Tokens.Num() - 1]; // Last token (skip -bm multiplier if present)
+				if (!FPaths::IsRelative(TexturePath))
+				{
+					CurrentMaterial->NormalTexture = TexturePath;
+				}
+				else
+				{
+					CurrentMaterial->NormalTexture = FPaths::Combine(MTLDirectory, TexturePath);
+				}
+			}
+			else if (Command == TEXT("map_ks") && Tokens.Num() >= 2)
+			{
+				FString TexturePath = Tokens[1];
+				if (!FPaths::IsRelative(TexturePath))
+				{
+					CurrentMaterial->SpecularTexture = TexturePath;
+				}
+				else
+				{
+					CurrentMaterial->SpecularTexture = FPaths::Combine(MTLDirectory, TexturePath);
+				}
+			}
+		}
+	}
+	
+	return Materials;
+}
+
+// Parse OBJ file
+static FOBJParseResult ParseOBJFile(const FString& OBJPath, float Scale = 1.0f)
+{
+	FOBJParseResult Result;
+	
+	FString OBJContent;
+	if (!FFileHelper::LoadFileToString(OBJContent, *OBJPath))
+	{
+		Result.ErrorMessage = FString::Printf(TEXT("Failed to read OBJ file: %s"), *OBJPath);
+		return Result;
+	}
+	
+	TArray<FString> Lines;
+	OBJContent.ParseIntoArrayLines(Lines);
+	
+	FString OBJDirectory = FPaths::GetPath(OBJPath);
+	
+	// Initialize with a default object
+	Result.Objects.Add(FOBJObject());
+	Result.Objects[0].Name = FPaths::GetBaseFilename(OBJPath);
+	Result.Objects[0].MaterialGroups.Add(FOBJMaterialGroup());
+	
+	FOBJObject* CurrentObject = &Result.Objects[0];
+	FOBJMaterialGroup* CurrentMaterialGroup = &CurrentObject->MaterialGroups[0];
+	
+	for (const FString& Line : Lines)
+	{
+		FString TrimmedLine = Line.TrimStartAndEnd();
+		if (TrimmedLine.IsEmpty() || TrimmedLine.StartsWith(TEXT("#")))
+		{
+			continue;
+		}
+		
+		TArray<FString> Tokens;
+		TrimmedLine.ParseIntoArrayWS(Tokens);
+		
+		if (Tokens.Num() == 0)
+		{
+			continue;
+		}
+		
+		FString Command = Tokens[0].ToLower();
+		
+		if (Command == TEXT("v") && Tokens.Num() >= 4)
+		{
+			// Vertex position - convert from OBJ (Y-up, right-handed) to UE (Z-up, left-handed)
+			FOBJVertex Vertex;
+			float X = FCString::Atof(*Tokens[1]);
+			float Y = FCString::Atof(*Tokens[2]);
+			float Z = FCString::Atof(*Tokens[3]);
+			// OBJ is right-handed Y-up, UE is left-handed Z-up
+			// Conversion: UE.X = OBJ.X, UE.Y = -OBJ.Z, UE.Z = OBJ.Y
+			Vertex.Position = FVector3f(X * Scale, -Z * Scale, Y * Scale);
+			Result.Vertices.Add(Vertex);
+		}
+		else if (Command == TEXT("vt") && Tokens.Num() >= 3)
+		{
+			// Texture coordinate
+			FOBJTexCoord TexCoord;
+			TexCoord.UV.X = FCString::Atof(*Tokens[1]);
+			TexCoord.UV.Y = 1.0f - FCString::Atof(*Tokens[2]); // Flip V for UE
+			Result.TexCoords.Add(TexCoord);
+		}
+		else if (Command == TEXT("vn") && Tokens.Num() >= 4)
+		{
+			// Vertex normal - convert coordinate system
+			FOBJNormal Normal;
+			float X = FCString::Atof(*Tokens[1]);
+			float Y = FCString::Atof(*Tokens[2]);
+			float Z = FCString::Atof(*Tokens[3]);
+			Normal.Normal = FVector3f(X, -Z, Y);
+			Normal.Normal.Normalize();
+			Result.Normals.Add(Normal);
+		}
+		else if (Command == TEXT("f") && Tokens.Num() >= 4)
+		{
+			// Face (triangulate if needed)
+			TArray<FOBJFaceVertex> FaceVertices;
+			for (int32 i = 1; i < Tokens.Num(); i++)
+			{
+				FaceVertices.Add(ParseFaceVertex(Tokens[i]));
+			}
+			
+			// Triangulate the face (fan triangulation)
+			for (int32 i = 2; i < FaceVertices.Num(); i++)
+			{
+				FOBJFace Triangle;
+				Triangle.Vertices.Add(FaceVertices[0]);
+				Triangle.Vertices.Add(FaceVertices[i - 1]);
+				Triangle.Vertices.Add(FaceVertices[i]);
+				CurrentMaterialGroup->Faces.Add(Triangle);
+			}
+		}
+		else if (Command == TEXT("o") && Tokens.Num() >= 2)
+		{
+			// New object
+			Result.Objects.Add(FOBJObject());
+			CurrentObject = &Result.Objects.Last();
+			CurrentObject->Name = Tokens[1];
+			CurrentObject->MaterialGroups.Add(FOBJMaterialGroup());
+			CurrentMaterialGroup = &CurrentObject->MaterialGroups[0];
+		}
+		else if (Command == TEXT("g") && Tokens.Num() >= 2)
+		{
+			// Group - treat similarly to object for material grouping purposes
+			if (CurrentObject->MaterialGroups.Last().Faces.Num() > 0)
+			{
+				CurrentObject->MaterialGroups.Add(FOBJMaterialGroup());
+				CurrentMaterialGroup = &CurrentObject->MaterialGroups.Last();
+			}
+		}
+		else if (Command == TEXT("usemtl") && Tokens.Num() >= 2)
+		{
+			// Use material
+			FString MaterialName = Tokens[1];
+			
+			// Check if we need a new material group or if current one is empty
+			if (CurrentMaterialGroup->Faces.Num() > 0 || !CurrentMaterialGroup->MaterialName.IsEmpty())
+			{
+				CurrentObject->MaterialGroups.Add(FOBJMaterialGroup());
+				CurrentMaterialGroup = &CurrentObject->MaterialGroups.Last();
+			}
+			CurrentMaterialGroup->MaterialName = MaterialName;
+		}
+		else if (Command == TEXT("mtllib") && Tokens.Num() >= 2)
+		{
+			// Material library
+			Result.MTLLibrary = Tokens[1];
+			FString MTLPath = FPaths::Combine(OBJDirectory, Result.MTLLibrary);
+			Result.Materials = ParseMTLFile(MTLPath);
+		}
+	}
+	
+	// Clean up empty objects and material groups
+	for (int32 i = Result.Objects.Num() - 1; i >= 0; i--)
+	{
+		FOBJObject& Obj = Result.Objects[i];
+		for (int32 j = Obj.MaterialGroups.Num() - 1; j >= 0; j--)
+		{
+			if (Obj.MaterialGroups[j].Faces.Num() == 0)
+			{
+				Obj.MaterialGroups.RemoveAt(j);
+			}
+		}
+		if (Obj.MaterialGroups.Num() == 0)
+		{
+			Result.Objects.RemoveAt(i);
+		}
+	}
+	
+	Result.bSuccess = true;
+	return Result;
+}
+
+// Build MeshDescription from parsed OBJ data
+static void BuildMeshDescriptionFromOBJ(
+	FMeshDescription& MeshDescription,
+	const FOBJParseResult& ParseResult,
+	bool bCombineMeshes,
+	TArray<FString>& OutMaterialNames)
+{
+	FStaticMeshAttributes Attributes(MeshDescription);
+	Attributes.Register();
+	
+	TVertexAttributesRef<FVector3f> VertexPositions = Attributes.GetVertexPositions();
+	TVertexInstanceAttributesRef<FVector3f> VertexInstanceNormals = Attributes.GetVertexInstanceNormals();
+	TVertexInstanceAttributesRef<FVector3f> VertexInstanceTangents = Attributes.GetVertexInstanceTangents();
+	TVertexInstanceAttributesRef<float> VertexInstanceBinormalSigns = Attributes.GetVertexInstanceBinormalSigns();
+	TVertexInstanceAttributesRef<FVector2f> VertexInstanceUVs = Attributes.GetVertexInstanceUVs();
+	TVertexInstanceAttributesRef<FVector4f> VertexInstanceColors = Attributes.GetVertexInstanceColors();
+	TPolygonGroupAttributesRef<FName> PolygonGroupImportedMaterialSlotNames = Attributes.GetPolygonGroupMaterialSlotNames();
+	
+	// Reserve vertex positions
+	MeshDescription.ReserveNewVertices(ParseResult.Vertices.Num());
+	
+	// Create all vertices
+	TArray<FVertexID> VertexIDs;
+	VertexIDs.Reserve(ParseResult.Vertices.Num());
+	for (const FOBJVertex& Vertex : ParseResult.Vertices)
+	{
+		FVertexID VertexID = MeshDescription.CreateVertex();
+		VertexPositions[VertexID] = Vertex.Position;
+		VertexIDs.Add(VertexID);
+	}
+	
+	// Track material slots
+	TMap<FString, FPolygonGroupID> MaterialToPolygonGroup;
+	int32 MaterialIndex = 0;
+	
+	// Process all objects
+	for (const FOBJObject& Object : ParseResult.Objects)
+	{
+		for (const FOBJMaterialGroup& MaterialGroup : Object.MaterialGroups)
+		{
+			// Get or create polygon group for this material
+			FPolygonGroupID PolygonGroupID;
+			FString MaterialName = MaterialGroup.MaterialName.IsEmpty() ? TEXT("DefaultMaterial") : MaterialGroup.MaterialName;
+			
+			if (FPolygonGroupID* ExistingGroup = MaterialToPolygonGroup.Find(MaterialName))
+			{
+				PolygonGroupID = *ExistingGroup;
+			}
+			else
+			{
+				PolygonGroupID = MeshDescription.CreatePolygonGroup();
+				PolygonGroupImportedMaterialSlotNames[PolygonGroupID] = FName(*MaterialName);
+				MaterialToPolygonGroup.Add(MaterialName, PolygonGroupID);
+				OutMaterialNames.Add(MaterialName);
+				MaterialIndex++;
+			}
+			
+			// Add all faces
+			for (const FOBJFace& Face : MaterialGroup.Faces)
+			{
+				TArray<FVertexInstanceID> VertexInstanceIDs;
+				VertexInstanceIDs.Reserve(Face.Vertices.Num());
+				
+				for (const FOBJFaceVertex& FaceVertex : Face.Vertices)
+				{
+					// Handle negative indices (relative to current position)
+					int32 VertexIndex = FaceVertex.PositionIndex;
+					if (VertexIndex < 0)
+					{
+						VertexIndex = ParseResult.Vertices.Num() + VertexIndex + 1;
+					}
+					
+					if (VertexIndex < 1 || VertexIndex > VertexIDs.Num())
+					{
+						continue; // Invalid vertex index
+					}
+					
+					FVertexID VertexID = VertexIDs[VertexIndex - 1]; // OBJ indices are 1-based
+					FVertexInstanceID VertexInstanceID = MeshDescription.CreateVertexInstance(VertexID);
+					
+					// Set UV
+					if (FaceVertex.TexCoordIndex != 0 && FMath::Abs(FaceVertex.TexCoordIndex) <= ParseResult.TexCoords.Num())
+					{
+						int32 UVIndex = FaceVertex.TexCoordIndex;
+						if (UVIndex < 0)
+						{
+							UVIndex = ParseResult.TexCoords.Num() + UVIndex + 1;
+						}
+						VertexInstanceUVs.Set(VertexInstanceID, 0, ParseResult.TexCoords[UVIndex - 1].UV);
+					}
+					else
+					{
+						VertexInstanceUVs.Set(VertexInstanceID, 0, FVector2f::ZeroVector);
+					}
+					
+					// Set normal
+					if (FaceVertex.NormalIndex != 0 && FMath::Abs(FaceVertex.NormalIndex) <= ParseResult.Normals.Num())
+					{
+						int32 NormalIndex = FaceVertex.NormalIndex;
+						if (NormalIndex < 0)
+						{
+							NormalIndex = ParseResult.Normals.Num() + NormalIndex + 1;
+						}
+						VertexInstanceNormals[VertexInstanceID] = ParseResult.Normals[NormalIndex - 1].Normal;
+					}
+					else
+					{
+						VertexInstanceNormals[VertexInstanceID] = FVector3f::UpVector;
+					}
+					
+					// Set default tangent and binormal
+					VertexInstanceTangents[VertexInstanceID] = FVector3f(1.0f, 0.0f, 0.0f);
+					VertexInstanceBinormalSigns[VertexInstanceID] = 1.0f;
+					
+					// Set default vertex color
+					VertexInstanceColors[VertexInstanceID] = FVector4f(1.0f, 1.0f, 1.0f, 1.0f);
+					
+					VertexInstanceIDs.Add(VertexInstanceID);
+				}
+				
+				// Create the polygon (triangle)
+				if (VertexInstanceIDs.Num() >= 3)
+				{
+					MeshDescription.CreatePolygon(PolygonGroupID, VertexInstanceIDs);
+				}
+			}
+		}
+	}
+}
+
+FECACommandResult FECACommand_ImportOBJ::Execute(const TSharedPtr<FJsonObject>& Params)
+{
+	FString SourcePath;
+	if (!GetStringParam(Params, TEXT("source_path"), SourcePath))
+	{
+		return FECACommandResult::Error(TEXT("Missing required parameter: source_path"));
+	}
+	
+	FString DestinationPath;
+	if (!GetStringParam(Params, TEXT("destination_path"), DestinationPath))
+	{
+		return FECACommandResult::Error(TEXT("Missing required parameter: destination_path"));
+	}
+	
+	FString MeshName;
+	GetStringParam(Params, TEXT("mesh_name"), MeshName, false);
+	
+	double Scale = 1.0;
+	GetFloatParam(Params, TEXT("scale"), Scale, false);
+	
+	bool bImportMaterials = true;
+	GetBoolParam(Params, TEXT("import_materials"), bImportMaterials, false);
+	
+	bool bGenerateLightmapUVs = true;
+	GetBoolParam(Params, TEXT("generate_lightmap_uvs"), bGenerateLightmapUVs, false);
+	
+	bool bAutoGenerateCollision = true;
+	GetBoolParam(Params, TEXT("auto_generate_collision"), bAutoGenerateCollision, false);
+	
+	bool bCombineMeshes = true;
+	GetBoolParam(Params, TEXT("combine_meshes"), bCombineMeshes, false);
+	
+	// Check if source file exists
+	if (!FPaths::FileExists(SourcePath))
+	{
+		return FECACommandResult::Error(FString::Printf(TEXT("Source file does not exist: %s"), *SourcePath));
+	}
+	
+	// Determine asset name
+	if (MeshName.IsEmpty())
+	{
+		MeshName = FPaths::GetBaseFilename(SourcePath);
+	}
+	
+	// Ensure destination path ends with /
+	if (!DestinationPath.EndsWith(TEXT("/")))
+	{
+		DestinationPath += TEXT("/");
+	}
+	
+	// Parse the OBJ file
+	FOBJParseResult ParseResult = ParseOBJFile(SourcePath, static_cast<float>(Scale));
+	if (!ParseResult.bSuccess)
+	{
+		return FECACommandResult::Error(ParseResult.ErrorMessage);
+	}
+	
+	// Validate we have geometry
+	if (ParseResult.Vertices.Num() == 0)
+	{
+		return FECACommandResult::Error(TEXT("OBJ file contains no vertices"));
+	}
+	
+	int32 TotalFaces = 0;
+	for (const FOBJObject& Obj : ParseResult.Objects)
+	{
+		for (const FOBJMaterialGroup& Group : Obj.MaterialGroups)
+		{
+			TotalFaces += Group.Faces.Num();
+		}
+	}
+	
+	if (TotalFaces == 0)
+	{
+		return FECACommandResult::Error(TEXT("OBJ file contains no faces"));
+	}
+	
+	// Create the package path
+	FString PackagePath = DestinationPath + MeshName;
+	
+	// Check if asset already exists at this path
+	if (UObject* ExistingAsset = StaticLoadObject(UObject::StaticClass(), nullptr, *PackagePath))
+	{
+		if (UStaticMesh* ExistingMesh = Cast<UStaticMesh>(ExistingAsset))
+		{
+			TSharedPtr<FJsonObject> Result = MakeResult();
+			Result->SetStringField(TEXT("mesh_path"), PackagePath);
+			Result->SetStringField(TEXT("mesh_name"), MeshName);
+			Result->SetBoolField(TEXT("already_exists"), true);
+			Result->SetStringField(TEXT("message"), TEXT("Static Mesh already exists at this path"));
+			return FECACommandResult::Success(Result);
+		}
+		else
+		{
+			return FECACommandResult::Error(FString::Printf(TEXT("Asset already exists at path '%s' but is not a Static Mesh (it's a %s)"), *PackagePath, *ExistingAsset->GetClass()->GetName()));
+		}
+	}
+	
+	// Create the package
+	UPackage* Package = CreatePackage(*PackagePath);
+	if (!Package)
+	{
+		return FECACommandResult::Error(TEXT("Failed to create package"));
+	}
+	
+	// Create the static mesh
+	UStaticMesh* StaticMesh = NewObject<UStaticMesh>(Package, *MeshName, RF_Public | RF_Standalone);
+	if (!StaticMesh)
+	{
+		return FECACommandResult::Error(TEXT("Failed to create static mesh"));
+	}
+	
+	// Build mesh description
+	FMeshDescription MeshDesc;
+	TArray<FString> MaterialNames;
+	BuildMeshDescriptionFromOBJ(MeshDesc, ParseResult, bCombineMeshes, MaterialNames);
+	
+	// Configure static mesh build settings
+	StaticMesh->GetStaticMaterials().Empty();
+	for (int32 i = 0; i < MaterialNames.Num(); i++)
+	{
+		StaticMesh->GetStaticMaterials().Add(FStaticMaterial(nullptr, FName(*MaterialNames[i]), FName(*MaterialNames[i])));
+	}
+	
+	// Add LOD 0
+	FStaticMeshSourceModel& SourceModel = StaticMesh->AddSourceModel();
+	SourceModel.BuildSettings.bRecomputeNormals = ParseResult.Normals.Num() == 0;
+	SourceModel.BuildSettings.bRecomputeTangents = true;
+	SourceModel.BuildSettings.bGenerateLightmapUVs = bGenerateLightmapUVs;
+	SourceModel.BuildSettings.SrcLightmapIndex = 0;
+	SourceModel.BuildSettings.DstLightmapIndex = 1;
+	
+	// Create mesh description for LOD 0
+	FMeshDescription* MeshDescription = StaticMesh->CreateMeshDescription(0);
+	*MeshDescription = MoveTemp(MeshDesc);
+	StaticMesh->CommitMeshDescription(0);
+	
+	// Configure collision
+	if (bAutoGenerateCollision)
+	{
+		StaticMesh->CreateBodySetup();
+		if (StaticMesh->GetBodySetup())
+		{
+			StaticMesh->GetBodySetup()->CollisionTraceFlag = CTF_UseComplexAsSimple;
+		}
+	}
+	
+	// Build the mesh
+	StaticMesh->SetImportVersion(EImportStaticMeshVersion::LastVersion);
+	StaticMesh->Build(false);
+	StaticMesh->PostEditChange();
+	
+	// Notify asset registry
+	FAssetRegistryModule::AssetCreated(StaticMesh);
+	Package->MarkPackageDirty();
+	
+	// Build result
+	TSharedPtr<FJsonObject> Result = MakeResult();
+	Result->SetStringField(TEXT("mesh_path"), StaticMesh->GetPathName());
+	Result->SetStringField(TEXT("mesh_name"), StaticMesh->GetName());
+	Result->SetNumberField(TEXT("vertex_count"), ParseResult.Vertices.Num());
+	Result->SetNumberField(TEXT("triangle_count"), TotalFaces);
+	Result->SetNumberField(TEXT("material_count"), MaterialNames.Num());
+	
+	// Add material slot info
+	TArray<TSharedPtr<FJsonValue>> MaterialsArray;
+	for (int32 i = 0; i < MaterialNames.Num(); i++)
+	{
+		TSharedPtr<FJsonObject> MatObj = MakeShared<FJsonObject>();
+		MatObj->SetNumberField(TEXT("slot_index"), i);
+		MatObj->SetStringField(TEXT("slot_name"), MaterialNames[i]);
+		MaterialsArray.Add(MakeShared<FJsonValueObject>(MatObj));
+	}
+	Result->SetArrayField(TEXT("material_slots"), MaterialsArray);
+	
+	// Add object info if we had multiple objects
+	if (ParseResult.Objects.Num() > 0)
+	{
+		TArray<TSharedPtr<FJsonValue>> ObjectsArray;
+		for (const FOBJObject& Obj : ParseResult.Objects)
+		{
+			TSharedPtr<FJsonObject> ObjInfo = MakeShared<FJsonObject>();
+			ObjInfo->SetStringField(TEXT("name"), Obj.Name);
+			int32 ObjFaces = 0;
+			for (const FOBJMaterialGroup& Group : Obj.MaterialGroups)
+			{
+				ObjFaces += Group.Faces.Num();
+			}
+			ObjInfo->SetNumberField(TEXT("triangle_count"), ObjFaces);
+			ObjectsArray.Add(MakeShared<FJsonValueObject>(ObjInfo));
+		}
+		Result->SetArrayField(TEXT("objects"), ObjectsArray);
+	}
+	
+	// Include MTL library info if present
+	if (!ParseResult.MTLLibrary.IsEmpty())
+	{
+		Result->SetStringField(TEXT("mtl_library"), ParseResult.MTLLibrary);
+		Result->SetNumberField(TEXT("mtl_material_count"), ParseResult.Materials.Num());
+	}
+	
+	return FECACommandResult::Success(Result);
+}
+
+//==============================================================================
+// EXPORT TEXTURE
+//==============================================================================
+
+FECACommandResult FECACommand_ExportTexture::Execute(const TSharedPtr<FJsonObject>& Params)
+{
+	FString TexturePath;
+	if (!GetStringParam(Params, TEXT("texture_path"), TexturePath))
+	{
+		return FECACommandResult::Error(TEXT("Missing required parameter: texture_path"));
+	}
+
+	FString OutputPath;
+	if (!GetStringParam(Params, TEXT("output_path"), OutputPath))
+	{
+		return FECACommandResult::Error(TEXT("Missing required parameter: output_path"));
+	}
+
+	int32 MipLevel = 0;
+	GetIntParam(Params, TEXT("mip_level"), MipLevel, false);
+
+	// Load the texture
+	UTexture2D* Texture = LoadObject<UTexture2D>(nullptr, *TexturePath);
+	if (!Texture)
+	{
+		return FECACommandResult::Error(FString::Printf(TEXT("Texture not found: %s"), *TexturePath));
+	}
+
+	// Get source data from the texture
+	FTextureSource& Source = Texture->Source;
+	if (!Source.IsValid())
+	{
+		return FECACommandResult::Error(TEXT("Texture has no source data (may be a runtime texture)"));
+	}
+
+	int32 Width = Source.GetSizeX() >> MipLevel;
+	int32 Height = Source.GetSizeY() >> MipLevel;
+
+	if (Width == 0 || Height == 0)
+	{
+		return FECACommandResult::Error(TEXT("Texture has invalid dimensions at requested mip level"));
+	}
+
+	// Get the raw mip data
+	TArray64<uint8> SourceData;
+	if (!Source.GetMipData(SourceData, MipLevel))
+	{
+		return FECACommandResult::Error(TEXT("Failed to get texture mip data"));
+	}
+
+	// Convert to FColor array based on source format
+	TArray<FColor> Bitmap;
+	Bitmap.SetNum(Width * Height);
+
+	ETextureSourceFormat SourceFormat = Source.GetFormat();
+	
+	if (SourceFormat == TSF_BGRA8)
+	{
+		const uint8* PixelData = SourceData.GetData();
+		for (int32 i = 0; i < Width * Height; i++)
+		{
+			Bitmap[i] = FColor(PixelData[i * 4 + 2], PixelData[i * 4 + 1], PixelData[i * 4 + 0], PixelData[i * 4 + 3]);
+		}
+	}
+	else if (SourceFormat == TSF_G8)
+	{
+		const uint8* PixelData = SourceData.GetData();
+		for (int32 i = 0; i < Width * Height; i++)
+		{
+			uint8 Gray = PixelData[i];
+			Bitmap[i] = FColor(Gray, Gray, Gray, 255);
+		}
+	}
+	else if (SourceFormat == TSF_RGBA16F)
+	{
+		const FFloat16* PixelData = reinterpret_cast<const FFloat16*>(SourceData.GetData());
+		for (int32 i = 0; i < Width * Height; i++)
+		{
+			FLinearColor Linear(PixelData[i * 4 + 0].GetFloat(), PixelData[i * 4 + 1].GetFloat(), PixelData[i * 4 + 2].GetFloat(), PixelData[i * 4 + 3].GetFloat());
+			Bitmap[i] = Linear.ToFColor(true);
+		}
+	}
+	else
+	{
+		return FECACommandResult::Error(FString::Printf(TEXT("Unsupported texture source format: %d"), (int32)SourceFormat));
+	}
+
+	// Compress to PNG
+	TArray64<uint8> CompressedData;
+	FImageUtils::PNGCompressImageArray(Width, Height, TArrayView64<const FColor>(Bitmap.GetData(), Bitmap.Num()), CompressedData);
+
+	// Save to file
+	if (!FFileHelper::SaveArrayToFile(CompressedData, *OutputPath))
+	{
+		return FECACommandResult::Error(FString::Printf(TEXT("Failed to save PNG to: %s"), *OutputPath));
+	}
+
+	TSharedPtr<FJsonObject> Result = MakeShared<FJsonObject>();
+	Result->SetStringField(TEXT("texture_path"), TexturePath);
+	Result->SetStringField(TEXT("output_path"), OutputPath);
+	Result->SetNumberField(TEXT("width"), Width);
+	Result->SetNumberField(TEXT("height"), Height);
+	Result->SetNumberField(TEXT("mip_level"), MipLevel);
+	Result->SetNumberField(TEXT("file_size"), CompressedData.Num());
+	Result->SetStringField(TEXT("source_format"), UEnum::GetValueAsString(SourceFormat));
+
+	return FECACommandResult::Success(Result);
+}
+
+//==============================================================================
+// EXPORT RENDER TARGET
+//==============================================================================
+
+FECACommandResult FECACommand_ExportRenderTarget::Execute(const TSharedPtr<FJsonObject>& Params)
+{
+	FString RenderTargetPath;
+	if (!GetStringParam(Params, TEXT("render_target_path"), RenderTargetPath))
+	{
+		return FECACommandResult::Error(TEXT("Missing required parameter: render_target_path"));
+	}
+
+	FString OutputPath;
+	if (!GetStringParam(Params, TEXT("output_path"), OutputPath))
+	{
+		return FECACommandResult::Error(TEXT("Missing required parameter: output_path"));
+	}
+
+	// Load the render target
+	UTextureRenderTarget2D* RenderTarget = LoadObject<UTextureRenderTarget2D>(nullptr, *RenderTargetPath);
+	if (!RenderTarget)
+	{
+		return FECACommandResult::Error(FString::Printf(TEXT("Render target not found: %s"), *RenderTargetPath));
+	}
+
+	int32 Width = RenderTarget->SizeX;
+	int32 Height = RenderTarget->SizeY;
+
+	// Read pixels from render target
+	FTextureRenderTargetResource* RTResource = RenderTarget->GameThread_GetRenderTargetResource();
+	if (!RTResource)
+	{
+		return FECACommandResult::Error(TEXT("Failed to get render target resource"));
+	}
+
+	TArray<FColor> Bitmap;
+	Bitmap.SetNum(Width * Height);
+
+	// Try reading as FColor first
+	if (!RTResource->ReadPixels(Bitmap))
+	{
+		// Try reading as float and converting
+		TArray<FFloat16Color> FloatPixels;
+		RTResource->ReadFloat16Pixels(FloatPixels);
+		
+		if (FloatPixels.Num() == Width * Height)
+		{
+			for (int32 i = 0; i < FloatPixels.Num(); i++)
+			{
+				FLinearColor Linear(FloatPixels[i].R.GetFloat(), FloatPixels[i].G.GetFloat(), FloatPixels[i].B.GetFloat(), FloatPixels[i].A.GetFloat());
+				Bitmap[i] = Linear.ToFColor(true);
+			}
+		}
+		else
+		{
+			return FECACommandResult::Error(TEXT("Failed to read render target pixels"));
+		}
+	}
+
+	// Compress to PNG
+	TArray64<uint8> CompressedData;
+	FImageUtils::PNGCompressImageArray(Width, Height, TArrayView64<const FColor>(Bitmap.GetData(), Bitmap.Num()), CompressedData);
+
+	// Save to file
+	if (!FFileHelper::SaveArrayToFile(CompressedData, *OutputPath))
+	{
+		return FECACommandResult::Error(FString::Printf(TEXT("Failed to save PNG to: %s"), *OutputPath));
+	}
+
+	TSharedPtr<FJsonObject> Result = MakeShared<FJsonObject>();
+	Result->SetStringField(TEXT("render_target_path"), RenderTargetPath);
+	Result->SetStringField(TEXT("output_path"), OutputPath);
+	Result->SetNumberField(TEXT("width"), Width);
+	Result->SetNumberField(TEXT("height"), Height);
+	Result->SetNumberField(TEXT("file_size"), CompressedData.Num());
+
+	return FECACommandResult::Success(Result);
+}
+
+//------------------------------------------------------------------------------
+// Generic Asset Property Commands (Reflection-based)
+//------------------------------------------------------------------------------
+
+// Helper to convert JSON value to string for property setting
+static FString JsonValueToPropertyString(const TSharedPtr<FJsonObject>& Params, const FString& Key)
+{
+	if (!Params->HasField(Key))
+	{
+		return TEXT("");
+	}
+	
+	const TSharedPtr<FJsonValue> Value = Params->TryGetField(Key);
+	
+	switch (Value->Type)
+	{
+		case EJson::Boolean:
+			return Value->AsBool() ? TEXT("True") : TEXT("False");
+		case EJson::Number:
+			return FString::SanitizeFloat(Value->AsNumber());
+		case EJson::String:
+			return Value->AsString();
+		default:
+			return TEXT("");
+	}
+}
+
+// Helper to convert FProperty value to JSON
+static TSharedPtr<FJsonValue> PropertyToJsonValue(FProperty* Property, const void* ValuePtr)
+{
+	if (FBoolProperty* BoolProp = CastField<FBoolProperty>(Property))
+	{
+		return MakeShared<FJsonValueBoolean>(BoolProp->GetPropertyValue(ValuePtr));
+	}
+	else if (FNumericProperty* NumProp = CastField<FNumericProperty>(Property))
+	{
+		if (NumProp->IsFloatingPoint())
+		{
+			return MakeShared<FJsonValueNumber>(NumProp->GetFloatingPointPropertyValue(ValuePtr));
+		}
+		else if (NumProp->IsInteger())
+		{
+			return MakeShared<FJsonValueNumber>(static_cast<double>(NumProp->GetSignedIntPropertyValue(ValuePtr)));
+		}
+	}
+	else if (FEnumProperty* EnumProp = CastField<FEnumProperty>(Property))
+	{
+		UEnum* Enum = EnumProp->GetEnum();
+		FNumericProperty* UnderlyingProp = EnumProp->GetUnderlyingProperty();
+		int64 EnumValue = UnderlyingProp->GetSignedIntPropertyValue(ValuePtr);
+		FString EnumName = Enum->GetNameStringByValue(EnumValue);
+		return MakeShared<FJsonValueString>(EnumName);
+	}
+	else if (FByteProperty* ByteProp = CastField<FByteProperty>(Property))
+	{
+		if (UEnum* Enum = ByteProp->GetIntPropertyEnum())
+		{
+			uint8 ByteValue = *static_cast<const uint8*>(ValuePtr);
+			FString EnumName = Enum->GetNameStringByValue(ByteValue);
+			return MakeShared<FJsonValueString>(EnumName);
+		}
+		else
+		{
+			return MakeShared<FJsonValueNumber>(static_cast<double>(*static_cast<const uint8*>(ValuePtr)));
+		}
+	}
+	else if (FStrProperty* StrProp = CastField<FStrProperty>(Property))
+	{
+		return MakeShared<FJsonValueString>(StrProp->GetPropertyValue(ValuePtr));
+	}
+	else if (FNameProperty* NameProp = CastField<FNameProperty>(Property))
+	{
+		return MakeShared<FJsonValueString>(NameProp->GetPropertyValue(ValuePtr).ToString());
+	}
+	else if (FStructProperty* StructProp = CastField<FStructProperty>(Property))
+	{
+		// For structs, export as nested object
+		TSharedPtr<FJsonObject> StructObj = MakeShared<FJsonObject>();
+		for (TFieldIterator<FProperty> PropIt(StructProp->Struct); PropIt; ++PropIt)
+		{
+			FProperty* InnerProp = *PropIt;
+			const void* InnerValuePtr = InnerProp->ContainerPtrToValuePtr<void>(ValuePtr);
+			StructObj->SetField(InnerProp->GetName(), PropertyToJsonValue(InnerProp, InnerValuePtr));
+		}
+		return MakeShared<FJsonValueObject>(StructObj);
+	}
+	
+	// Fallback: export to string
+	FString StringValue;
+	Property->ExportTextItem_Direct(StringValue, ValuePtr, nullptr, nullptr, PPF_None);
+	return MakeShared<FJsonValueString>(StringValue);
+}
+
+// Find property by name, supporting common aliases (snake_case, nested paths, etc.)
+static FProperty* FindPropertyByPath(UClass* Class, const FString& PropertyPath, void*& OutContainer, UObject* Object)
+{
+	OutContainer = Object;
+	
+	// Handle nested property paths like "NaniteSettings.bEnabled"
+	TArray<FString> PathParts;
+	PropertyPath.ParseIntoArray(PathParts, TEXT("."));
+	
+	FProperty* CurrentProperty = nullptr;
+	UStruct* CurrentStruct = Class;
+	
+	for (int32 i = 0; i < PathParts.Num(); i++)
+	{
+		FString PartName = PathParts[i];
+		
+		// Try to find the property
+		CurrentProperty = CurrentStruct->FindPropertyByName(FName(*PartName));
+		
+		// Try with 'b' prefix for bools
+		if (!CurrentProperty)
+		{
+			CurrentProperty = CurrentStruct->FindPropertyByName(FName(*(TEXT("b") + PartName)));
+		}
+		
+		// Convert snake_case to PascalCase
+		if (!CurrentProperty)
+		{
+			FString PascalCase;
+			bool bCapitalizeNext = true;
+			for (TCHAR Char : PartName)
+			{
+				if (Char == '_')
+				{
+					bCapitalizeNext = true;
+				}
+				else
+				{
+					PascalCase += bCapitalizeNext ? FChar::ToUpper(Char) : Char;
+					bCapitalizeNext = false;
+				}
+			}
+			CurrentProperty = CurrentStruct->FindPropertyByName(FName(*PascalCase));
+			
+			// Try with 'b' prefix
+			if (!CurrentProperty)
+			{
+				CurrentProperty = CurrentStruct->FindPropertyByName(FName(*(TEXT("b") + PascalCase)));
+			}
+		}
+		
+		if (!CurrentProperty)
+		{
+			return nullptr;
+		}
+		
+		// If this isn't the last part, we need to traverse into a struct
+		if (i < PathParts.Num() - 1)
+		{
+			FStructProperty* StructProp = CastField<FStructProperty>(CurrentProperty);
+			if (!StructProp)
+			{
+				return nullptr; // Can't traverse into non-struct
+			}
+			
+			OutContainer = StructProp->ContainerPtrToValuePtr<void>(OutContainer);
+			CurrentStruct = StructProp->Struct;
+		}
+	}
+	
+	return CurrentProperty;
+}
+
+FECACommandResult FECACommand_SetAssetProperty::Execute(const TSharedPtr<FJsonObject>& Params)
+{
+	FString AssetPath = Params->GetStringField(TEXT("asset_path"));
+	FString PropertyPath = Params->GetStringField(TEXT("property"));
+	
+	// Load the asset
+	UObject* Asset = LoadObject<UObject>(nullptr, *AssetPath);
+	if (!Asset)
+	{
+		return FECACommandResult::Error(FString::Printf(TEXT("Failed to load asset: %s"), *AssetPath));
+	}
+	
+	// Find the property using reflection
+	void* Container = nullptr;
+	FProperty* Property = FindPropertyByPath(Asset->GetClass(), PropertyPath, Container, Asset);
+	if (!Property)
+	{
+		return FECACommandResult::Error(FString::Printf(TEXT("Property '%s' not found on %s"), *PropertyPath, *Asset->GetClass()->GetName()));
+	}
+	
+	// Get the value as a string for ImportText
+	FString ValueStr = JsonValueToPropertyString(Params, TEXT("value"));
+	
+	// Get pointer to property value
+	void* ValuePtr = Property->ContainerPtrToValuePtr<void>(Container);
+	
+	// Store old value for reporting
+	TSharedPtr<FJsonValue> OldValue = PropertyToJsonValue(Property, ValuePtr);
+	
+	// Use ImportText to set the value
+	const TCHAR* ImportResult = Property->ImportText_Direct(*ValueStr, ValuePtr, Asset, PPF_None);
+	if (!ImportResult)
+	{
+		return FECACommandResult::Error(FString::Printf(TEXT("Failed to set property '%s' to value '%s'"), *PropertyPath, *ValueStr));
+	}
+	
+	// Get new value for reporting
+	TSharedPtr<FJsonValue> NewValue = PropertyToJsonValue(Property, ValuePtr);
+	
+	// Mark asset as modified
+	Asset->PostEditChange();
+	Asset->MarkPackageDirty();
+	
+	TSharedPtr<FJsonObject> Result = MakeShared<FJsonObject>();
+	Result->SetStringField(TEXT("asset_path"), AssetPath);
+	Result->SetStringField(TEXT("asset_class"), Asset->GetClass()->GetName());
+	Result->SetStringField(TEXT("property"), Property->GetName());
+	Result->SetField(TEXT("old_value"), OldValue);
+	Result->SetField(TEXT("new_value"), NewValue);
+	Result->SetBoolField(TEXT("success"), true);
+	
+	return FECACommandResult::Success(Result);
+}
+
+FECACommandResult FECACommand_GetAssetProperty::Execute(const TSharedPtr<FJsonObject>& Params)
+{
+	FString AssetPath = Params->GetStringField(TEXT("asset_path"));
+	FString PropertyPath = Params->HasField(TEXT("property")) ? Params->GetStringField(TEXT("property")) : TEXT("");
+	
+	// Load the asset
+	UObject* Asset = LoadObject<UObject>(nullptr, *AssetPath);
+	if (!Asset)
+	{
+		return FECACommandResult::Error(FString::Printf(TEXT("Failed to load asset: %s"), *AssetPath));
+	}
+	
+	TSharedPtr<FJsonObject> Result = MakeShared<FJsonObject>();
+	Result->SetStringField(TEXT("asset_path"), AssetPath);
+	Result->SetStringField(TEXT("asset_class"), Asset->GetClass()->GetName());
+	
+	if (!PropertyPath.IsEmpty())
+	{
+		// Get specific property
+		void* Container = nullptr;
+		FProperty* Property = FindPropertyByPath(Asset->GetClass(), PropertyPath, Container, Asset);
+		if (!Property)
+		{
+			return FECACommandResult::Error(FString::Printf(TEXT("Property '%s' not found on %s"), *PropertyPath, *Asset->GetClass()->GetName()));
+		}
+		
+		void* ValuePtr = Property->ContainerPtrToValuePtr<void>(Container);
+		Result->SetField(Property->GetName(), PropertyToJsonValue(Property, ValuePtr));
+	}
+	else
+	{
+		// Return all editable properties
+		TSharedPtr<FJsonObject> PropertiesObj = MakeShared<FJsonObject>();
+		
+		for (TFieldIterator<FProperty> PropIt(Asset->GetClass()); PropIt; ++PropIt)
+		{
+			FProperty* Property = *PropIt;
+			
+			// Only include EditAnywhere properties
+			if (!Property->HasAnyPropertyFlags(CPF_Edit))
+			{
+				continue;
+			}
+			
+			// Skip deprecated and transient properties
+			if (Property->HasAnyPropertyFlags(CPF_Deprecated | CPF_Transient))
+			{
+				continue;
+			}
+			
+			void* ValuePtr = Property->ContainerPtrToValuePtr<void>(Asset);
+			PropertiesObj->SetField(Property->GetName(), PropertyToJsonValue(Property, ValuePtr));
+		}
+		
+		Result->SetObjectField(TEXT("properties"), PropertiesObj);
+	}
+	
+	return FECACommandResult::Success(Result);
+}
+
+FECACommandResult FECACommand_ListAssetProperties::Execute(const TSharedPtr<FJsonObject>& Params)
+{
+	FString AssetPath = Params->GetStringField(TEXT("asset_path"));
+	FString CategoryFilter = Params->HasField(TEXT("category")) ? Params->GetStringField(TEXT("category")) : TEXT("");
+	FString SearchFilter = Params->HasField(TEXT("search")) ? Params->GetStringField(TEXT("search")).ToLower() : TEXT("");
+	
+	// Load the asset
+	UObject* Asset = LoadObject<UObject>(nullptr, *AssetPath);
+	if (!Asset)
+	{
+		return FECACommandResult::Error(FString::Printf(TEXT("Failed to load asset: %s"), *AssetPath));
+	}
+	
+	TSharedPtr<FJsonObject> Result = MakeShared<FJsonObject>();
+	Result->SetStringField(TEXT("asset_path"), AssetPath);
+	Result->SetStringField(TEXT("asset_class"), Asset->GetClass()->GetName());
+	
+	// Organize properties by category
+	TMap<FString, TArray<TSharedPtr<FJsonValue>>> CategorizedProperties;
+	
+	for (TFieldIterator<FProperty> PropIt(Asset->GetClass()); PropIt; ++PropIt)
+	{
+		FProperty* Property = *PropIt;
+		
+		// Only include EditAnywhere properties
+		if (!Property->HasAnyPropertyFlags(CPF_Edit))
+		{
+			continue;
+		}
+		
+		// Skip deprecated properties
+		if (Property->HasAnyPropertyFlags(CPF_Deprecated))
+		{
+			continue;
+		}
+		
+		// Apply search filter
+		if (!SearchFilter.IsEmpty() && !Property->GetName().ToLower().Contains(SearchFilter))
+		{
+			continue;
+		}
+		
+		// Get category from metadata
+		FString Category = Property->GetMetaData(TEXT("Category"));
+		if (Category.IsEmpty())
+		{
+			Category = TEXT("Uncategorized");
+		}
+		
+		// Apply category filter
+		if (!CategoryFilter.IsEmpty() && !Category.Contains(CategoryFilter))
+		{
+			continue;
+		}
+		
+		// Build property info
+		TSharedPtr<FJsonObject> PropInfo = MakeShared<FJsonObject>();
+		PropInfo->SetStringField(TEXT("name"), Property->GetName());
+		PropInfo->SetStringField(TEXT("type"), Property->GetCPPType());
+		
+		// Get display name if available
+		FString DisplayName = Property->GetMetaData(TEXT("DisplayName"));
+		if (!DisplayName.IsEmpty())
+		{
+			PropInfo->SetStringField(TEXT("display_name"), DisplayName);
+		}
+		
+		// Get tooltip if available
+		FString ToolTip = Property->GetMetaData(TEXT("ToolTip"));
+		if (!ToolTip.IsEmpty())
+		{
+			PropInfo->SetStringField(TEXT("tooltip"), ToolTip);
+		}
+		
+		// Check if it's an enum and list values
+		if (FEnumProperty* EnumProp = CastField<FEnumProperty>(Property))
+		{
+			UEnum* Enum = EnumProp->GetEnum();
+			TArray<TSharedPtr<FJsonValue>> EnumValues;
+			for (int32 i = 0; i < Enum->NumEnums() - 1; i++) // -1 to skip _MAX
+			{
+				EnumValues.Add(MakeShared<FJsonValueString>(Enum->GetNameStringByIndex(i)));
+			}
+			PropInfo->SetArrayField(TEXT("enum_values"), EnumValues);
+		}
+		else if (FByteProperty* ByteProp = CastField<FByteProperty>(Property))
+		{
+			if (UEnum* Enum = ByteProp->GetIntPropertyEnum())
+			{
+				TArray<TSharedPtr<FJsonValue>> EnumValues;
+				for (int32 i = 0; i < Enum->NumEnums() - 1; i++)
+				{
+					EnumValues.Add(MakeShared<FJsonValueString>(Enum->GetNameStringByIndex(i)));
+				}
+				PropInfo->SetArrayField(TEXT("enum_values"), EnumValues);
+			}
+		}
+		
+		// Check for struct properties
+		if (FStructProperty* StructProp = CastField<FStructProperty>(Property))
+		{
+			TArray<TSharedPtr<FJsonValue>> StructFields;
+			for (TFieldIterator<FProperty> InnerIt(StructProp->Struct); InnerIt; ++InnerIt)
+			{
+				TSharedPtr<FJsonObject> FieldInfo = MakeShared<FJsonObject>();
+				FieldInfo->SetStringField(TEXT("name"), InnerIt->GetName());
+				FieldInfo->SetStringField(TEXT("type"), InnerIt->GetCPPType());
+				StructFields.Add(MakeShared<FJsonValueObject>(FieldInfo));
+			}
+			PropInfo->SetArrayField(TEXT("struct_fields"), StructFields);
+		}
+		
+		CategorizedProperties.FindOrAdd(Category).Add(MakeShared<FJsonValueObject>(PropInfo));
+	}
+	
+	// Convert to JSON object organized by category
+	TSharedPtr<FJsonObject> CategoriesObj = MakeShared<FJsonObject>();
+	for (auto& Pair : CategorizedProperties)
+	{
+		CategoriesObj->SetArrayField(Pair.Key, Pair.Value);
+	}
+	Result->SetObjectField(TEXT("properties_by_category"), CategoriesObj);
+	Result->SetNumberField(TEXT("total_properties"), [&]() {
+		int32 Count = 0;
+		for (auto& Pair : CategorizedProperties) Count += Pair.Value.Num();
+		return Count;
+	}());
+	
+	return FECACommandResult::Success(Result);
+}
+
+//------------------------------------------------------------------------------
+// Delete Asset
+//------------------------------------------------------------------------------
+
+FECACommandResult FECACommand_DeleteAsset::Execute(const TSharedPtr<FJsonObject>& Params)
+{
+	FString AssetPath;
+	if (!GetStringParam(Params, TEXT("asset_path"), AssetPath))
+	{
+		return FECACommandResult::Error(TEXT("Missing required parameter: asset_path"));
+	}
+	
+	bool bForce = false;
+	GetBoolParam(Params, TEXT("force"), bForce, false);
+	
+	// Load the asset to make sure it exists
+	UObject* Asset = LoadObject<UObject>(nullptr, *AssetPath);
+	if (!Asset)
+	{
+		return FECACommandResult::Error(FString::Printf(TEXT("Asset not found: %s"), *AssetPath));
+	}
+	
+	// Get asset tools
+	FAssetToolsModule& AssetToolsModule = FModuleManager::LoadModuleChecked<FAssetToolsModule>("AssetTools");
+	
+	// Get the object path
+	FString PackagePath = Asset->GetOutermost()->GetName();
+	
+	// Check for references if not forcing
+	if (!bForce)
+	{
+		TArray<FAssetIdentifier> Referencers;
+		FAssetRegistryModule& AssetRegistryModule = FModuleManager::LoadModuleChecked<FAssetRegistryModule>("AssetRegistry");
+		AssetRegistryModule.Get().GetReferencers(FAssetIdentifier(FName(*PackagePath)), Referencers);
+		
+		if (Referencers.Num() > 0)
+		{
+			TSharedPtr<FJsonObject> Result = MakeResult();
+			Result->SetBoolField(TEXT("success"), false);
+			Result->SetStringField(TEXT("error"), TEXT("Asset has references. Use force=true to delete anyway."));
+			Result->SetNumberField(TEXT("reference_count"), Referencers.Num());
+			
+			TArray<TSharedPtr<FJsonValue>> RefArray;
+			for (const FAssetIdentifier& Ref : Referencers)
+			{
+				RefArray.Add(MakeShared<FJsonValueString>(Ref.ToString()));
+			}
+			Result->SetArrayField(TEXT("references"), RefArray);
+			
+			return FECACommandResult::Success(Result);
+		}
+	}
+	
+	// Delete the asset
+	TArray<UObject*> AssetsToDelete;
+	AssetsToDelete.Add(Asset);
+	
+	int32 DeletedCount = ObjectTools::DeleteObjects(AssetsToDelete, /*bShowConfirmation=*/false);
+	
+	if (DeletedCount > 0)
+	{
+		TSharedPtr<FJsonObject> Result = MakeResult();
+		Result->SetBoolField(TEXT("success"), true);
+		Result->SetStringField(TEXT("deleted_asset"), AssetPath);
+		return FECACommandResult::Success(Result);
+	}
+	else
+	{
+		return FECACommandResult::Error(TEXT("Failed to delete asset"));
+	}
+}
+
+//------------------------------------------------------------------------------
+// Rename Asset
+//------------------------------------------------------------------------------
+
+FECACommandResult FECACommand_RenameAsset::Execute(const TSharedPtr<FJsonObject>& Params)
+{
+	FString AssetPath;
+	if (!GetStringParam(Params, TEXT("asset_path"), AssetPath))
+	{
+		return FECACommandResult::Error(TEXT("Missing required parameter: asset_path"));
+	}
+	
+	FString NewName;
+	if (!GetStringParam(Params, TEXT("new_name"), NewName))
+	{
+		return FECACommandResult::Error(TEXT("Missing required parameter: new_name"));
+	}
+	
+	// Load the asset
+	UObject* Asset = LoadObject<UObject>(nullptr, *AssetPath);
+	if (!Asset)
+	{
+		return FECACommandResult::Error(FString::Printf(TEXT("Asset not found: %s"), *AssetPath));
+	}
+	
+	// Get the current package path (folder)
+	FString PackagePath = FPackageName::GetLongPackagePath(Asset->GetOutermost()->GetName());
+	FString NewPath = PackagePath / NewName;
+	
+	// Use asset tools to rename
+	FAssetToolsModule& AssetToolsModule = FModuleManager::LoadModuleChecked<FAssetToolsModule>("AssetTools");
+	
+	TArray<FAssetRenameData> RenameData;
+	RenameData.Add(FAssetRenameData(Asset, PackagePath, NewName));
+	
+	bool bSuccess = AssetToolsModule.Get().RenameAssets(RenameData);
+	
+	if (bSuccess)
+	{
+		TSharedPtr<FJsonObject> Result = MakeResult();
+		Result->SetBoolField(TEXT("success"), true);
+		Result->SetStringField(TEXT("old_path"), AssetPath);
+		Result->SetStringField(TEXT("new_path"), NewPath);
+		return FECACommandResult::Success(Result);
+	}
+	else
+	{
+		return FECACommandResult::Error(TEXT("Failed to rename asset"));
+	}
+}
+
+//------------------------------------------------------------------------------
+// Move Asset
+//------------------------------------------------------------------------------
+
+FECACommandResult FECACommand_MoveAsset::Execute(const TSharedPtr<FJsonObject>& Params)
+{
+	FString AssetPath;
+	if (!GetStringParam(Params, TEXT("asset_path"), AssetPath))
+	{
+		return FECACommandResult::Error(TEXT("Missing required parameter: asset_path"));
+	}
+	
+	FString DestinationPath;
+	if (!GetStringParam(Params, TEXT("destination_path"), DestinationPath))
+	{
+		return FECACommandResult::Error(TEXT("Missing required parameter: destination_path"));
+	}
+	
+	// Load the asset
+	UObject* Asset = LoadObject<UObject>(nullptr, *AssetPath);
+	if (!Asset)
+	{
+		return FECACommandResult::Error(FString::Printf(TEXT("Asset not found: %s"), *AssetPath));
+	}
+	
+	// Get the asset name
+	FString AssetName = Asset->GetName();
+	FString NewPath = DestinationPath / AssetName;
+	
+	// Use asset tools to rename/move
+	FAssetToolsModule& AssetToolsModule = FModuleManager::LoadModuleChecked<FAssetToolsModule>("AssetTools");
+	
+	TArray<FAssetRenameData> RenameData;
+	RenameData.Add(FAssetRenameData(Asset, DestinationPath, AssetName));
+	
+	bool bSuccess = AssetToolsModule.Get().RenameAssets(RenameData);
+	
+	if (bSuccess)
+	{
+		TSharedPtr<FJsonObject> Result = MakeResult();
+		Result->SetBoolField(TEXT("success"), true);
+		Result->SetStringField(TEXT("old_path"), AssetPath);
+		Result->SetStringField(TEXT("new_path"), NewPath);
+		return FECACommandResult::Success(Result);
+	}
+	else
+	{
+		return FECACommandResult::Error(TEXT("Failed to move asset"));
+	}
+}
+
+//------------------------------------------------------------------------------
+// Duplicate Asset
+//------------------------------------------------------------------------------
+
+FECACommandResult FECACommand_DuplicateAsset::Execute(const TSharedPtr<FJsonObject>& Params)
+{
+	FString AssetPath;
+	if (!GetStringParam(Params, TEXT("asset_path"), AssetPath))
+	{
+		return FECACommandResult::Error(TEXT("Missing required parameter: asset_path"));
+	}
+	
+	FString NewName;
+	if (!GetStringParam(Params, TEXT("new_name"), NewName))
+	{
+		return FECACommandResult::Error(TEXT("Missing required parameter: new_name"));
+	}
+	
+	FString DestinationPath;
+	GetStringParam(Params, TEXT("destination_path"), DestinationPath, false);
+	
+	// Load the asset
+	UObject* Asset = LoadObject<UObject>(nullptr, *AssetPath);
+	if (!Asset)
+	{
+		return FECACommandResult::Error(FString::Printf(TEXT("Asset not found: %s"), *AssetPath));
+	}
+	
+	// If no destination specified, use same folder as source
+	if (DestinationPath.IsEmpty())
+	{
+		DestinationPath = FPackageName::GetLongPackagePath(Asset->GetOutermost()->GetName());
+	}
+	
+	FString NewPath = DestinationPath / NewName;
+	
+	// Use asset tools to duplicate
+	FAssetToolsModule& AssetToolsModule = FModuleManager::LoadModuleChecked<FAssetToolsModule>("AssetTools");
+	
+	UObject* DuplicatedAsset = AssetToolsModule.Get().DuplicateAsset(NewName, DestinationPath, Asset);
+	
+	if (DuplicatedAsset)
+	{
+		TSharedPtr<FJsonObject> Result = MakeResult();
+		Result->SetBoolField(TEXT("success"), true);
+		Result->SetStringField(TEXT("source_path"), AssetPath);
+		Result->SetStringField(TEXT("new_path"), NewPath);
+		Result->SetStringField(TEXT("new_asset_class"), DuplicatedAsset->GetClass()->GetName());
+		return FECACommandResult::Success(Result);
+	}
+	else
+	{
+		return FECACommandResult::Error(TEXT("Failed to duplicate asset"));
+	}
+}
+
+//------------------------------------------------------------------------------
+// ImportFBX - FBX File Importer using Interchange Framework
+//------------------------------------------------------------------------------
+
+FECACommandResult FECACommand_ImportFBX::Execute(const TSharedPtr<FJsonObject>& Params)
+{
+	// Get required parameters
+	FString SourcePath;
+	if (!GetStringParam(Params, TEXT("source_path"), SourcePath))
+	{
+		return FECACommandResult::Error(TEXT("Missing required parameter: source_path"));
+	}
+	
+	FString DestinationPath;
+	if (!GetStringParam(Params, TEXT("destination_path"), DestinationPath))
+	{
+		return FECACommandResult::Error(TEXT("Missing required parameter: destination_path"));
+	}
+	
+	// Get optional parameters
+	FString AssetName;
+	GetStringParam(Params, TEXT("asset_name"), AssetName, false);
+	
+	bool bImportMaterials = true;
+	GetBoolParam(Params, TEXT("import_materials"), bImportMaterials, false);
+	
+	bool bImportTextures = true;
+	GetBoolParam(Params, TEXT("import_textures"), bImportTextures, false);
+	
+	bool bImportAsSkeletal = false;
+	GetBoolParam(Params, TEXT("import_as_skeletal"), bImportAsSkeletal, false);
+	
+	bool bImportAnimations = false;
+	GetBoolParam(Params, TEXT("import_animations"), bImportAnimations, false);
+	
+	bool bGenerateLightmapUVs = true;
+	GetBoolParam(Params, TEXT("generate_lightmap_uvs"), bGenerateLightmapUVs, false);
+	
+	bool bAutoGenerateCollision = true;
+	GetBoolParam(Params, TEXT("auto_generate_collision"), bAutoGenerateCollision, false);
+	
+	bool bCombineMeshes = true;
+	GetBoolParam(Params, TEXT("combine_meshes"), bCombineMeshes, false);
+	
+	bool bConvertSceneUnit = true;
+	GetBoolParam(Params, TEXT("convert_scene_unit"), bConvertSceneUnit, false);
+	
+	bool bForceFrontXAxis = false;
+	GetBoolParam(Params, TEXT("force_front_x_axis"), bForceFrontXAxis, false);
+	
+	// Check if source file exists
+	if (!FPaths::FileExists(SourcePath))
+	{
+		return FECACommandResult::Error(FString::Printf(TEXT("Source file does not exist: %s"), *SourcePath));
+	}
+	
+	// Verify it's an FBX file
+	FString Extension = FPaths::GetExtension(SourcePath).ToLower();
+	if (Extension != TEXT("fbx"))
+	{
+		return FECACommandResult::Error(FString::Printf(TEXT("File is not an FBX file (extension: %s). Expected .fbx"), *Extension));
+	}
+	
+	// Determine asset name
+	if (AssetName.IsEmpty())
+	{
+		AssetName = FPaths::GetBaseFilename(SourcePath);
+	}
+	
+	// Ensure destination path ends with /
+	if (!DestinationPath.EndsWith(TEXT("/")))
+	{
+		DestinationPath += TEXT("/");
+	}
+	
+	// Check if Interchange is enabled
+	if (!UInterchangeManager::IsInterchangeImportEnabled())
+	{
+		return FECACommandResult::Error(TEXT("Interchange import is disabled. Enable it in Project Settings > Interchange"));
+	}
+	
+	// Get the Interchange manager
+	UInterchangeManager& InterchangeManager = UInterchangeManager::GetInterchangeManager();
+	
+	// Create source data from the file
+	UInterchangeSourceData* SourceData = UInterchangeManager::CreateSourceData(SourcePath);
+	if (!SourceData)
+	{
+		return FECACommandResult::Error(TEXT("Failed to create Interchange source data"));
+	}
+	
+	// Check if we can translate this file
+	if (!InterchangeManager.CanTranslateSourceData(SourceData))
+	{
+		return FECACommandResult::Error(TEXT("No translator available for this FBX file. Ensure Interchange FBX translator is enabled."));
+	}
+	
+	// Set up import parameters
+	FImportAssetParameters ImportAssetParameters;
+	ImportAssetParameters.bIsAutomated = true;  // Don't show UI dialogs
+	ImportAssetParameters.bReplaceExisting = true;
+	ImportAssetParameters.DestinationName = AssetName;
+	
+	// Track imported objects for the result
+	TArray<UObject*> ImportedObjects;
+	bool bImportStarted = false;
+	bool bImportComplete = false;
+	FString ImportError;
+	
+	// Set up completion callback
+	ImportAssetParameters.OnAssetsImportDoneNative.BindLambda([&ImportedObjects, &bImportComplete](const TArray<UObject*>& Objects)
+	{
+		ImportedObjects = Objects;
+		bImportComplete = true;
+	});
+	
+	// Perform synchronous import
+	bImportStarted = InterchangeManager.ImportAsset(DestinationPath, SourceData, ImportAssetParameters, ImportedObjects);
+	
+	if (!bImportStarted)
+	{
+		return FECACommandResult::Error(TEXT("Failed to start FBX import"));
+	}
+	
+	// Build result
+	TSharedPtr<FJsonObject> Result = MakeResult();
+	Result->SetBoolField(TEXT("success"), true);
+	Result->SetStringField(TEXT("source_file"), SourcePath);
+	Result->SetStringField(TEXT("destination_path"), DestinationPath);
+	Result->SetNumberField(TEXT("imported_count"), ImportedObjects.Num());
+	
+	// List all imported objects
+	TArray<TSharedPtr<FJsonValue>> ImportedArray;
+	int32 StaticMeshCount = 0;
+	int32 SkeletalMeshCount = 0;
+	int32 MaterialCount = 0;
+	int32 TextureCount = 0;
+	int32 AnimationCount = 0;
+	
+	for (UObject* ImportedObject : ImportedObjects)
+	{
+		if (!ImportedObject) continue;
+		
+		TSharedPtr<FJsonObject> ObjInfo = MakeShared<FJsonObject>();
+		ObjInfo->SetStringField(TEXT("name"), ImportedObject->GetName());
+		ObjInfo->SetStringField(TEXT("path"), ImportedObject->GetPathName());
+		ObjInfo->SetStringField(TEXT("class"), ImportedObject->GetClass()->GetName());
+		
+		// Count by type
+		if (Cast<UStaticMesh>(ImportedObject))
+		{
+			StaticMeshCount++;
+			ObjInfo->SetStringField(TEXT("type"), TEXT("StaticMesh"));
+			
+			// Get mesh stats
+			UStaticMesh* Mesh = Cast<UStaticMesh>(ImportedObject);
+			if (Mesh->GetRenderData() && Mesh->GetRenderData()->LODResources.Num() > 0)
+			{
+				const FStaticMeshLODResources& LOD0 = Mesh->GetRenderData()->LODResources[0];
+				ObjInfo->SetNumberField(TEXT("vertex_count"), LOD0.GetNumVertices());
+				ObjInfo->SetNumberField(TEXT("triangle_count"), LOD0.GetNumTriangles());
+			}
+			ObjInfo->SetNumberField(TEXT("material_slots"), Mesh->GetStaticMaterials().Num());
+		}
+		else if (Cast<USkeletalMesh>(ImportedObject))
+		{
+			SkeletalMeshCount++;
+			ObjInfo->SetStringField(TEXT("type"), TEXT("SkeletalMesh"));
+		}
+		else if (Cast<UMaterialInterface>(ImportedObject))
+		{
+			MaterialCount++;
+			ObjInfo->SetStringField(TEXT("type"), TEXT("Material"));
+		}
+		else if (Cast<UTexture>(ImportedObject))
+		{
+			TextureCount++;
+			ObjInfo->SetStringField(TEXT("type"), TEXT("Texture"));
+		}
+		else if (ImportedObject->GetClass()->GetName().Contains(TEXT("AnimSequence")))
+		{
+			AnimationCount++;
+			ObjInfo->SetStringField(TEXT("type"), TEXT("Animation"));
+		}
+		else
+		{
+			ObjInfo->SetStringField(TEXT("type"), TEXT("Other"));
+		}
+		
+		ImportedArray.Add(MakeShared<FJsonValueObject>(ObjInfo));
+	}
+	
+	Result->SetArrayField(TEXT("imported_objects"), ImportedArray);
+	
+	// Summary counts
+	TSharedPtr<FJsonObject> Summary = MakeShared<FJsonObject>();
+	Summary->SetNumberField(TEXT("static_meshes"), StaticMeshCount);
+	Summary->SetNumberField(TEXT("skeletal_meshes"), SkeletalMeshCount);
+	Summary->SetNumberField(TEXT("materials"), MaterialCount);
+	Summary->SetNumberField(TEXT("textures"), TextureCount);
+	Summary->SetNumberField(TEXT("animations"), AnimationCount);
+	Summary->SetNumberField(TEXT("total"), ImportedObjects.Num());
+	Result->SetObjectField(TEXT("summary"), Summary);
+	
+	// Return the primary asset path (first static or skeletal mesh)
+	for (UObject* ImportedObject : ImportedObjects)
+	{
+		if (Cast<UStaticMesh>(ImportedObject) || Cast<USkeletalMesh>(ImportedObject))
+		{
+			Result->SetStringField(TEXT("primary_mesh_path"), ImportedObject->GetPathName());
+			Result->SetStringField(TEXT("primary_mesh_name"), ImportedObject->GetName());
+			break;
+		}
+	}
+	
+	return FECACommandResult::Success(Result);
+}
+
+
+//------------------------------------------------------------------------------
+// GetAssetThumbnail - Get thumbnail/icon for an asset
+//------------------------------------------------------------------------------
+
+// Helper function to extract thumbnail from a Texture2D asset directly
+static bool ExtractTextureThumbnail(UTexture2D* Texture, int32 Size, TArray<uint8>& OutPNGData, FString& OutError)
+{
+	// Get source data from the texture
+	FTextureSource& Source = Texture->Source;
+	if (!Source.IsValid())
+	{
+		OutError = TEXT("Texture has no source data (may be a runtime-only texture)");
+		return false;
+	}
+
+	int32 SrcWidth = Source.GetSizeX();
+	int32 SrcHeight = Source.GetSizeY();
+
+	if (SrcWidth == 0 || SrcHeight == 0)
+	{
+		OutError = TEXT("Texture has invalid dimensions");
+		return false;
+	}
+
+	// Get the raw mip data (mip 0 = full resolution)
+	TArray64<uint8> SourceData;
+	if (!Source.GetMipData(SourceData, 0))
+	{
+		OutError = TEXT("Failed to get texture mip data");
+		return false;
+	}
+
+	// Convert to FColor array based on source format
+	TArray<FColor> SrcBitmap;
+	SrcBitmap.SetNum(SrcWidth * SrcHeight);
+
+	ETextureSourceFormat SourceFormat = Source.GetFormat();
+
+	if (SourceFormat == TSF_BGRA8)
+	{
+		const uint8* PixelData = SourceData.GetData();
+		for (int32 i = 0; i < SrcWidth * SrcHeight; i++)
+		{
+			SrcBitmap[i] = FColor(PixelData[i * 4 + 2], PixelData[i * 4 + 1], PixelData[i * 4 + 0], PixelData[i * 4 + 3]);
+		}
+	}
+	else if (SourceFormat == TSF_G8)
+	{
+		const uint8* PixelData = SourceData.GetData();
+		for (int32 i = 0; i < SrcWidth * SrcHeight; i++)
+		{
+			uint8 Gray = PixelData[i];
+			SrcBitmap[i] = FColor(Gray, Gray, Gray, 255);
+		}
+	}
+	else if (SourceFormat == TSF_RGBA16F)
+	{
+		const FFloat16* PixelData = reinterpret_cast<const FFloat16*>(SourceData.GetData());
+		for (int32 i = 0; i < SrcWidth * SrcHeight; i++)
+		{
+			FLinearColor Linear(PixelData[i * 4 + 0].GetFloat(), PixelData[i * 4 + 1].GetFloat(), PixelData[i * 4 + 2].GetFloat(), PixelData[i * 4 + 3].GetFloat());
+			SrcBitmap[i] = Linear.ToFColor(true);
+		}
+	}
+	else
+	{
+		OutError = FString::Printf(TEXT("Unsupported texture source format: %d"), (int32)SourceFormat);
+		return false;
+	}
+
+	// Resize if needed
+	TArray<FColor> FinalBitmap;
+	int32 FinalWidth, FinalHeight;
+
+	if (SrcWidth == Size && SrcHeight == Size)
+	{
+		// No resize needed
+		FinalBitmap = MoveTemp(SrcBitmap);
+		FinalWidth = SrcWidth;
+		FinalHeight = SrcHeight;
+	}
+	else
+	{
+		// Calculate output size maintaining aspect ratio, fitting within Size x Size
+		float AspectRatio = (float)SrcWidth / (float)SrcHeight;
+		if (AspectRatio >= 1.0f)
+		{
+			FinalWidth = Size;
+			FinalHeight = FMath::Max(1, FMath::RoundToInt(Size / AspectRatio));
+		}
+		else
+		{
+			FinalHeight = Size;
+			FinalWidth = FMath::Max(1, FMath::RoundToInt(Size * AspectRatio));
+		}
+
+		// Simple bilinear resize
+		FinalBitmap.SetNum(FinalWidth * FinalHeight);
+		for (int32 y = 0; y < FinalHeight; y++)
+		{
+			for (int32 x = 0; x < FinalWidth; x++)
+			{
+				float SrcX = (float)x * (float)(SrcWidth - 1) / (float)(FinalWidth - 1);
+				float SrcY = (float)y * (float)(SrcHeight - 1) / (float)(FinalHeight - 1);
+
+				int32 X0 = FMath::FloorToInt(SrcX);
+				int32 Y0 = FMath::FloorToInt(SrcY);
+				int32 X1 = FMath::Min(X0 + 1, SrcWidth - 1);
+				int32 Y1 = FMath::Min(Y0 + 1, SrcHeight - 1);
+
+				float FracX = SrcX - X0;
+				float FracY = SrcY - Y0;
+
+				FColor C00 = SrcBitmap[Y0 * SrcWidth + X0];
+				FColor C10 = SrcBitmap[Y0 * SrcWidth + X1];
+				FColor C01 = SrcBitmap[Y1 * SrcWidth + X0];
+				FColor C11 = SrcBitmap[Y1 * SrcWidth + X1];
+
+				FColor Result;
+				Result.R = FMath::RoundToInt(FMath::Lerp(FMath::Lerp((float)C00.R, (float)C10.R, FracX), FMath::Lerp((float)C01.R, (float)C11.R, FracX), FracY));
+				Result.G = FMath::RoundToInt(FMath::Lerp(FMath::Lerp((float)C00.G, (float)C10.G, FracX), FMath::Lerp((float)C01.G, (float)C11.G, FracX), FracY));
+				Result.B = FMath::RoundToInt(FMath::Lerp(FMath::Lerp((float)C00.B, (float)C10.B, FracX), FMath::Lerp((float)C01.B, (float)C11.B, FracX), FracY));
+				Result.A = FMath::RoundToInt(FMath::Lerp(FMath::Lerp((float)C00.A, (float)C10.A, FracX), FMath::Lerp((float)C01.A, (float)C11.A, FracX), FracY));
+
+				FinalBitmap[y * FinalWidth + x] = Result;
+			}
+		}
+	}
+
+	// Compress to PNG
+	TArray64<uint8> CompressedData;
+	FImageUtils::PNGCompressImageArray(FinalWidth, FinalHeight, TArrayView64<const FColor>(FinalBitmap.GetData(), FinalBitmap.Num()), CompressedData);
+
+	// Copy to output
+	OutPNGData.SetNum(CompressedData.Num());
+	FMemory::Memcpy(OutPNGData.GetData(), CompressedData.GetData(), CompressedData.Num());
+
+	return true;
+}
+
+// Helper function to generate and encode a thumbnail for an asset
+static bool GenerateAssetThumbnail(UObject* Asset, int32 Size, TArray<uint8>& OutPNGData, FString& OutError)
+{
+	if (!Asset)
+	{
+		OutError = TEXT("Asset is null");
+		return false;
+	}
+
+	// Clamp size to reasonable limits
+	Size = FMath::Clamp(Size, 32, 1024);
+
+	// Special case: Texture2D - extract directly since they don't use thumbnail renderers
+	if (UTexture2D* Texture = Cast<UTexture2D>(Asset))
+	{
+		return ExtractTextureThumbnail(Texture, Size, OutPNGData, OutError);
+	}
+
+	// Get the thumbnail manager for other asset types
+	FThumbnailRenderingInfo* RenderInfo = UThumbnailManager::Get().GetRenderingInfo(Asset);
+	if (!RenderInfo || !RenderInfo->Renderer)
+	{
+		OutError = FString::Printf(TEXT("No thumbnail renderer available for asset type: %s"), *Asset->GetClass()->GetName());
+		return false;
+	}
+
+	// Create a render target to render into
+	UTextureRenderTarget2D* RenderTarget = NewObject<UTextureRenderTarget2D>();
+	RenderTarget->InitCustomFormat(Size, Size, PF_B8G8R8A8, false);
+	RenderTarget->UpdateResourceImmediate(true);
+
+	FTextureRenderTargetResource* RenderTargetResource = RenderTarget->GameThread_GetRenderTargetResource();
+	if (!RenderTargetResource)
+	{
+		OutError = TEXT("Failed to create render target resource");
+		return false;
+	}
+
+	// Create canvas for rendering - must be done on game thread as renderers may create preview scenes
+	FCanvas Canvas(RenderTargetResource, nullptr, FGameTime(), GMaxRHIFeatureLevel);
+	Canvas.Clear(FLinearColor::Black);
+
+	// Draw the thumbnail - this happens on game thread
+	RenderInfo->Renderer->Draw(Asset, 0, 0, Size, Size, RenderTargetResource, &Canvas, false);
+
+	// Flush the canvas to submit draw calls
+	Canvas.Flush_GameThread();
+
+	// Wait for rendering to complete
+	FlushRenderingCommands();
+
+	// Read pixels from the render target
+	TArray<FColor> Bitmap;
+	Bitmap.SetNum(Size * Size);
+
+	if (!RenderTargetResource->ReadPixels(Bitmap))
+	{
+		OutError = TEXT("Failed to read pixels from render target");
+		return false;
+	}
+
+	// Compress to PNG
+	TArray64<uint8> CompressedData;
+	FImageUtils::PNGCompressImageArray(Size, Size, TArrayView64<const FColor>(Bitmap.GetData(), Bitmap.Num()), CompressedData);
+
+	// Copy to output
+	OutPNGData.SetNum(CompressedData.Num());
+	FMemory::Memcpy(OutPNGData.GetData(), CompressedData.GetData(), CompressedData.Num());
+
+	return true;
+}
+
+FECACommandResult FECACommand_GetAssetThumbnail::Execute(const TSharedPtr<FJsonObject>& Params)
+{
+	// Get required parameters
+	FString AssetPath;
+	if (!GetStringParam(Params, TEXT("asset_path"), AssetPath))
+	{
+		return FECACommandResult::Error(TEXT("Missing required parameter: asset_path"));
+	}
+
+	FString OutputPath;
+	if (!GetStringParam(Params, TEXT("output_path"), OutputPath))
+	{
+		return FECACommandResult::Error(TEXT("Missing required parameter: output_path"));
+	}
+
+	// Get optional parameters
+	int32 Size = 256;
+	GetIntParam(Params, TEXT("size"), Size, false);
+	Size = FMath::Clamp(Size, 32, 1024);
+
+	// Load the asset
+	UObject* Asset = LoadObject<UObject>(nullptr, *AssetPath);
+	if (!Asset)
+	{
+		return FECACommandResult::Error(FString::Printf(TEXT("Failed to load asset: %s"), *AssetPath));
+	}
+
+	// Generate the thumbnail
+	TArray<uint8> PNGData;
+	FString Error;
+	if (!GenerateAssetThumbnail(Asset, Size, PNGData, Error))
+	{
+		return FECACommandResult::Error(FString::Printf(TEXT("Failed to generate thumbnail: %s"), *Error));
+	}
+
+	// Ensure output directory exists
+	FString OutputDir = FPaths::GetPath(OutputPath);
+	if (!OutputDir.IsEmpty())
+	{
+		IFileManager::Get().MakeDirectory(*OutputDir, true);
+	}
+
+	// Save to file
+	if (!FFileHelper::SaveArrayToFile(PNGData, *OutputPath))
+	{
+		return FECACommandResult::Error(FString::Printf(TEXT("Failed to save thumbnail to: %s"), *OutputPath));
+	}
+
+	// Build result
+	TSharedPtr<FJsonObject> Result = MakeShared<FJsonObject>();
+	Result->SetStringField(TEXT("asset_path"), AssetPath);
+	Result->SetStringField(TEXT("asset_name"), Asset->GetName());
+	Result->SetStringField(TEXT("asset_class"), Asset->GetClass()->GetName());
+	Result->SetStringField(TEXT("output_path"), OutputPath);
+	Result->SetNumberField(TEXT("thumbnail_size"), Size);
+	Result->SetNumberField(TEXT("file_size"), PNGData.Num());
+
+	return FECACommandResult::Success(Result);
+}
+
+//------------------------------------------------------------------------------
+// GetAssetThumbnails - Get thumbnails for multiple assets
+//------------------------------------------------------------------------------
+
+FECACommandResult FECACommand_GetAssetThumbnails::Execute(const TSharedPtr<FJsonObject>& Params)
+{
+	// Get required parameters
+	const TArray<TSharedPtr<FJsonValue>>* AssetPathsArray = nullptr;
+	if (!GetArrayParam(Params, TEXT("asset_paths"), AssetPathsArray))
+	{
+		return FECACommandResult::Error(TEXT("Missing required parameter: asset_paths"));
+	}
+
+	FString OutputDirectory;
+	if (!GetStringParam(Params, TEXT("output_directory"), OutputDirectory))
+	{
+		return FECACommandResult::Error(TEXT("Missing required parameter: output_directory"));
+	}
+
+	// Get optional parameters
+	int32 Size = 256;
+	GetIntParam(Params, TEXT("size"), Size, false);
+	Size = FMath::Clamp(Size, 32, 1024);
+
+	// Ensure output directory exists and has trailing slash
+	if (!OutputDirectory.EndsWith(TEXT("/")) && !OutputDirectory.EndsWith(TEXT("\\")))
+	{
+		OutputDirectory += TEXT("/");
+	}
+	IFileManager::Get().MakeDirectory(*OutputDirectory, true);
+
+	// Process each asset
+	TArray<TSharedPtr<FJsonValue>> ResultsArray;
+	int32 SuccessCount = 0;
+	int32 FailureCount = 0;
+
+	for (const TSharedPtr<FJsonValue>& PathValue : *AssetPathsArray)
+	{
+		FString AssetPath = PathValue->AsString();
+		TSharedPtr<FJsonObject> ItemResult = MakeShared<FJsonObject>();
+		ItemResult->SetStringField(TEXT("asset_path"), AssetPath);
+
+		// Load the asset
+		UObject* Asset = LoadObject<UObject>(nullptr, *AssetPath);
+		if (!Asset)
+		{
+			ItemResult->SetBoolField(TEXT("success"), false);
+			ItemResult->SetStringField(TEXT("error"), TEXT("Failed to load asset"));
+			FailureCount++;
+			ResultsArray.Add(MakeShared<FJsonValueObject>(ItemResult));
+			continue;
+		}
+
+		ItemResult->SetStringField(TEXT("asset_name"), Asset->GetName());
+		ItemResult->SetStringField(TEXT("asset_class"), Asset->GetClass()->GetName());
+
+		// Generate the thumbnail
+		TArray<uint8> PNGData;
+		FString Error;
+		if (!GenerateAssetThumbnail(Asset, Size, PNGData, Error))
+		{
+			ItemResult->SetBoolField(TEXT("success"), false);
+			ItemResult->SetStringField(TEXT("error"), Error);
+			FailureCount++;
+			ResultsArray.Add(MakeShared<FJsonValueObject>(ItemResult));
+			continue;
+		}
+
+		// Save to file
+		FString Filename = Asset->GetName() + TEXT(".png");
+		FString FullPath = OutputDirectory + Filename;
+
+		if (FFileHelper::SaveArrayToFile(PNGData, *FullPath))
+		{
+			ItemResult->SetBoolField(TEXT("success"), true);
+			ItemResult->SetStringField(TEXT("output_path"), FullPath);
+			ItemResult->SetNumberField(TEXT("thumbnail_size"), Size);
+			ItemResult->SetNumberField(TEXT("file_size"), PNGData.Num());
+			SuccessCount++;
+		}
+		else
+		{
+			ItemResult->SetBoolField(TEXT("success"), false);
+			ItemResult->SetStringField(TEXT("error"), FString::Printf(TEXT("Failed to save to: %s"), *FullPath));
+			FailureCount++;
+		}
+
+		ResultsArray.Add(MakeShared<FJsonValueObject>(ItemResult));
+	}
+
+	// Build final result
+	TSharedPtr<FJsonObject> Result = MakeShared<FJsonObject>();
+	Result->SetStringField(TEXT("output_directory"), OutputDirectory);
+	Result->SetArrayField(TEXT("thumbnails"), ResultsArray);
+	Result->SetNumberField(TEXT("total_count"), AssetPathsArray->Num());
+	Result->SetNumberField(TEXT("success_count"), SuccessCount);
+	Result->SetNumberField(TEXT("failure_count"), FailureCount);
+	Result->SetNumberField(TEXT("thumbnail_size"), Size);
+
+	return FECACommandResult::Success(Result);
+}
