ECABridge experimental MCP plugin

diff --git a/Engine/Plugins/Experimental/ECABridge/Source/ECABridge/Private/Commands/ECABlueprintCommands.cpp b/Engine/Plugins/Experimental/ECABridge/Source/ECABridge/Private/Commands/ECABlueprintCommands.cpp
new file mode 100644
--- /dev/null
+++ b/Engine/Plugins/Experimental/ECABridge/Source/ECABridge/Private/Commands/ECABlueprintCommands.cpp
@@ -0,0 +1,870 @@
+// Copyright Epic Games, Inc. All Rights Reserved.
+
+#include "Commands/ECABlueprintCommands.h"
+#include "Commands/ECACommand.h"
+#include "Editor.h"
+#include "Engine/World.h"
+#include "Engine/Blueprint.h"
+#include "Engine/BlueprintGeneratedClass.h"
+#include "Kismet2/KismetEditorUtilities.h"
+#include "Kismet2/BlueprintEditorUtils.h"
+#include "Engine/SimpleConstructionScript.h"
+#include "Engine/SCS_Node.h"
+#include "Components/StaticMeshComponent.h"
+#include "Components/PointLightComponent.h"
+#include "Components/SpotLightComponent.h"
+#include "Components/DirectionalLightComponent.h"
+#include "Camera/CameraComponent.h"
+#include "Components/SceneCaptureComponent2D.h"
+#include "Components/AudioComponent.h"
+#include "Components/BoxComponent.h"
+#include "Components/SphereComponent.h"
+#include "Components/CapsuleComponent.h"
+#include "Components/ArrowComponent.h"
+#include "Components/BillboardComponent.h"
+#include "GameFramework/Actor.h"
+#include "GameFramework/Pawn.h"
+#include "GameFramework/Character.h"
+#include "AssetRegistry/AssetRegistryModule.h"
+#include "EdGraphSchema_K2.h"
+#include "Subsystems/AssetEditorSubsystem.h"
+#include "UObject/UObjectIterator.h"
+
+// Register all blueprint commands
+REGISTER_ECA_COMMAND(FECACommand_CreateBlueprint)
+REGISTER_ECA_COMMAND(FECACommand_AddBlueprintComponent)
+REGISTER_ECA_COMMAND(FECACommand_CompileBlueprint)
+REGISTER_ECA_COMMAND(FECACommand_AddBlueprintVariable)
+REGISTER_ECA_COMMAND(FECACommand_SpawnBlueprintActor)
+REGISTER_ECA_COMMAND(FECACommand_GetBlueprintInfo)
+REGISTER_ECA_COMMAND(FECACommand_ListBlueprints)
+REGISTER_ECA_COMMAND(FECACommand_OpenBlueprintEditor)
+REGISTER_ECA_COMMAND(FECACommand_DeleteBlueprintFunction)
+
+//------------------------------------------------------------------------------
+// Helper functions
+//------------------------------------------------------------------------------
+
+static UClass* GetParentClassFromString(const FString& ParentClass, FString* OutError = nullptr)
+{
+	// Common shorthand names
+	if (ParentClass.Equals(TEXT("Actor"), ESearchCase::IgnoreCase))
+	{
+		return AActor::StaticClass();
+	}
+	else if (ParentClass.Equals(TEXT("Pawn"), ESearchCase::IgnoreCase))
+	{
+		return APawn::StaticClass();
+	}
+	else if (ParentClass.Equals(TEXT("Character"), ESearchCase::IgnoreCase))
+	{
+		return ACharacter::StaticClass();
+	}
+	
+	// Build list of paths to try for the class
+	TArray<FString> PathsToTry;
+	
+	// If it already looks like a full path, try it first
+	if (ParentClass.Contains(TEXT("/")) || ParentClass.Contains(TEXT(".")))
+	{
+		PathsToTry.Add(ParentClass);
+	}
+	
+	// Try common module paths - strip leading 'A' prefix if present for class name matching
+	FString ClassName = ParentClass;
+	FString ClassNameNoPrefix = ParentClass.StartsWith(TEXT("A")) ? ParentClass.RightChop(1) : ParentClass;
+	
+	// Geometry Script classes (GeneratedDynamicMeshActor, DynamicMeshActor)
+	PathsToTry.Add(FString::Printf(TEXT("/Script/GeometryScriptingEditor.%s"), *ClassName));
+	PathsToTry.Add(FString::Printf(TEXT("/Script/GeometryScriptingEditor.A%s"), *ClassNameNoPrefix));
+	PathsToTry.Add(FString::Printf(TEXT("/Script/GeometryFramework.%s"), *ClassName));
+	PathsToTry.Add(FString::Printf(TEXT("/Script/GeometryFramework.A%s"), *ClassNameNoPrefix));
+	
+	// Engine classes
+	PathsToTry.Add(FString::Printf(TEXT("/Script/Engine.%s"), *ClassName));
+	PathsToTry.Add(FString::Printf(TEXT("/Script/Engine.A%s"), *ClassNameNoPrefix));
+	
+	// Game framework classes  
+	PathsToTry.Add(FString::Printf(TEXT("/Script/GameFramework.%s"), *ClassName));
+	PathsToTry.Add(FString::Printf(TEXT("/Script/GameFramework.A%s"), *ClassNameNoPrefix));
+	
+	// Niagara classes
+	PathsToTry.Add(FString::Printf(TEXT("/Script/Niagara.%s"), *ClassName));
+	PathsToTry.Add(FString::Printf(TEXT("/Script/Niagara.A%s"), *ClassNameNoPrefix));
+	
+	// Paper2D classes
+	PathsToTry.Add(FString::Printf(TEXT("/Script/Paper2D.%s"), *ClassName));
+	PathsToTry.Add(FString::Printf(TEXT("/Script/Paper2D.A%s"), *ClassNameNoPrefix));
+	
+	// Try to find the class using FindObject first (faster, works for already-loaded classes)
+	for (const FString& Path : PathsToTry)
+	{
+		if (UClass* FoundClass = FindObject<UClass>(nullptr, *Path))
+		{
+			if (FoundClass->IsChildOf(AActor::StaticClass()))
+			{
+				return FoundClass;
+			}
+		}
+	}
+	
+	// Try LoadObject for each path (loads the class if not already loaded)
+	for (const FString& Path : PathsToTry)
+	{
+		if (UClass* LoadedClass = LoadObject<UClass>(nullptr, *Path))
+		{
+			if (LoadedClass->IsChildOf(AActor::StaticClass()))
+			{
+				return LoadedClass;
+			}
+		}
+	}
+	
+	// Last resort: try to find by iterating all loaded classes
+	for (TObjectIterator<UClass> It; It; ++It)
+	{
+		UClass* TestClass = *It;
+		if (TestClass->IsChildOf(AActor::StaticClass()))
+		{
+			FString TestClassName = TestClass->GetName();
+			if (TestClassName.Equals(ClassName, ESearchCase::IgnoreCase) ||
+			    TestClassName.Equals(FString::Printf(TEXT("A%s"), *ClassNameNoPrefix), ESearchCase::IgnoreCase) ||
+			    TestClassName.Equals(ClassNameNoPrefix, ESearchCase::IgnoreCase))
+			{
+				return TestClass;
+			}
+		}
+	}
+	
+	// Report error if requested
+	if (OutError)
+	{
+		*OutError = FString::Printf(TEXT("Could not find Actor class '%s'. Tried paths: %s"), 
+			*ParentClass, *FString::Join(PathsToTry, TEXT(", ")));
+	}
+	
+	return nullptr;  // Return nullptr instead of defaulting to AActor - let caller handle the error
+}
+
+static UClass* GetComponentClassFromString(const FString& ComponentType)
+{
+	if (ComponentType.Equals(TEXT("StaticMesh"), ESearchCase::IgnoreCase))
+	{
+		return UStaticMeshComponent::StaticClass();
+	}
+	else if (ComponentType.Equals(TEXT("PointLight"), ESearchCase::IgnoreCase))
+	{
+		return UPointLightComponent::StaticClass();
+	}
+	else if (ComponentType.Equals(TEXT("SpotLight"), ESearchCase::IgnoreCase))
+	{
+		return USpotLightComponent::StaticClass();
+	}
+	else if (ComponentType.Equals(TEXT("DirectionalLight"), ESearchCase::IgnoreCase))
+	{
+		return UDirectionalLightComponent::StaticClass();
+	}
+	else if (ComponentType.Equals(TEXT("Camera"), ESearchCase::IgnoreCase))
+	{
+		return UCameraComponent::StaticClass();
+	}
+	else if (ComponentType.Equals(TEXT("SceneCapture2D"), ESearchCase::IgnoreCase))
+	{
+		return USceneCaptureComponent2D::StaticClass();
+	}
+	else if (ComponentType.Equals(TEXT("Audio"), ESearchCase::IgnoreCase))
+	{
+		return UAudioComponent::StaticClass();
+	}
+	else if (ComponentType.Equals(TEXT("Box"), ESearchCase::IgnoreCase) || ComponentType.Equals(TEXT("BoxCollision"), ESearchCase::IgnoreCase))
+	{
+		return UBoxComponent::StaticClass();
+	}
+	else if (ComponentType.Equals(TEXT("Sphere"), ESearchCase::IgnoreCase) || ComponentType.Equals(TEXT("SphereCollision"), ESearchCase::IgnoreCase))
+	{
+		return USphereComponent::StaticClass();
+	}
+	else if (ComponentType.Equals(TEXT("Capsule"), ESearchCase::IgnoreCase) || ComponentType.Equals(TEXT("CapsuleCollision"), ESearchCase::IgnoreCase))
+	{
+		return UCapsuleComponent::StaticClass();
+	}
+	else if (ComponentType.Equals(TEXT("Arrow"), ESearchCase::IgnoreCase))
+	{
+		return UArrowComponent::StaticClass();
+	}
+	else if (ComponentType.Equals(TEXT("Billboard"), ESearchCase::IgnoreCase))
+	{
+		return UBillboardComponent::StaticClass();
+	}
+	else if (ComponentType.Equals(TEXT("Scene"), ESearchCase::IgnoreCase))
+	{
+		return USceneComponent::StaticClass();
+	}
+	else
+	{
+		// Try to load as a class path
+		return LoadObject<UClass>(nullptr, *ComponentType);
+	}
+}
+
+//------------------------------------------------------------------------------
+// CreateBlueprint
+//------------------------------------------------------------------------------
+
+FECACommandResult FECACommand_CreateBlueprint::Execute(const TSharedPtr<FJsonObject>& Params)
+{
+	FString BlueprintName;
+	if (!GetStringParam(Params, TEXT("blueprint_name"), BlueprintName))
+	{
+		return FECACommandResult::Error(TEXT("Missing required parameter: blueprint_name"));
+	}
+	
+	FString ParentClassStr = TEXT("Actor");
+	GetStringParam(Params, TEXT("parent_class"), ParentClassStr, false);
+	
+	FString Path = TEXT("/Game/Blueprints/");
+	GetStringParam(Params, TEXT("path"), Path, false);
+	
+	// Ensure path ends with /
+	if (!Path.EndsWith(TEXT("/")))
+	{
+		Path += TEXT("/");
+	}
+	
+	FString ClassError;
+	UClass* ParentClass = GetParentClassFromString(ParentClassStr, &ClassError);
+	
+	if (!ParentClass)
+	{
+		return FECACommandResult::Error(FString::Printf(TEXT("Invalid parent class '%s'. %s"), *ParentClassStr, *ClassError));
+	}
+	
+	// Create package path
+	FString PackagePath = Path + BlueprintName;
+	
+	// Check if asset already exists at this path
+	if (UObject* ExistingAsset = StaticLoadObject(UObject::StaticClass(), nullptr, *PackagePath))
+	{
+		// Asset already exists - return info about it instead of crashing
+		if (UBlueprint* ExistingBlueprint = Cast<UBlueprint>(ExistingAsset))
+		{
+			TSharedPtr<FJsonObject> Result = MakeResult();
+			Result->SetStringField(TEXT("blueprint_path"), PackagePath);
+			Result->SetStringField(TEXT("blueprint_name"), BlueprintName);
+			Result->SetStringField(TEXT("parent_class"), ExistingBlueprint->ParentClass ? ExistingBlueprint->ParentClass->GetName() : TEXT("Unknown"));
+			Result->SetBoolField(TEXT("already_exists"), true);
+			Result->SetStringField(TEXT("message"), TEXT("Blueprint already exists at this path"));
+			return FECACommandResult::Success(Result);
+		}
+		else
+		{
+			return FECACommandResult::Error(FString::Printf(TEXT("Asset already exists at path '%s' but is not a Blueprint (it's a %s)"), *PackagePath, *ExistingAsset->GetClass()->GetName()));
+		}
+	}
+	
+	// Create package
+	UPackage* Package = CreatePackage(*PackagePath);
+	if (!Package)
+	{
+		return FECACommandResult::Error(TEXT("Failed to create package"));
+	}
+	
+	// Create blueprint
+	UBlueprint* Blueprint = FKismetEditorUtilities::CreateBlueprint(
+		ParentClass,
+		Package,
+		*BlueprintName,
+		BPTYPE_Normal,
+		UBlueprint::StaticClass(),
+		UBlueprintGeneratedClass::StaticClass()
+	);
+	
+	if (!Blueprint)
+	{
+		return FECACommandResult::Error(TEXT("Failed to create Blueprint"));
+	}
+	
+	// Notify asset registry and mark dirty
+	FAssetRegistryModule::AssetCreated(Blueprint);
+	Package->MarkPackageDirty();
+	
+	// Compile
+	FKismetEditorUtilities::CompileBlueprint(Blueprint);
+	
+	TSharedPtr<FJsonObject> Result = MakeResult();
+	Result->SetStringField(TEXT("blueprint_path"), PackagePath);
+	Result->SetStringField(TEXT("blueprint_name"), BlueprintName);
+	Result->SetStringField(TEXT("parent_class"), ParentClass->GetName());
+	return FECACommandResult::Success(Result);
+}
+
+//------------------------------------------------------------------------------
+// AddBlueprintComponent
+//------------------------------------------------------------------------------
+
+FECACommandResult FECACommand_AddBlueprintComponent::Execute(const TSharedPtr<FJsonObject>& Params)
+{
+	FString BlueprintPath;
+	if (!GetStringParam(Params, TEXT("blueprint_path"), BlueprintPath))
+	{
+		return FECACommandResult::Error(TEXT("Missing required parameter: blueprint_path"));
+	}
+	
+	FString ComponentType;
+	if (!GetStringParam(Params, TEXT("component_type"), ComponentType))
+	{
+		return FECACommandResult::Error(TEXT("Missing required parameter: component_type"));
+	}
+	
+	FString ComponentName = TEXT("NewComponent");
+	GetStringParam(Params, TEXT("component_name"), ComponentName, false);
+	
+	UBlueprint* Blueprint = LoadBlueprintByPath(BlueprintPath);
+	if (!Blueprint)
+	{
+		return FECACommandResult::Error(FString::Printf(TEXT("Blueprint not found: %s"), *BlueprintPath));
+	}
+	
+	UClass* ComponentClass = GetComponentClassFromString(ComponentType);
+	if (!ComponentClass)
+	{
+		return FECACommandResult::Error(FString::Printf(TEXT("Unknown component type: %s"), *ComponentType));
+	}
+	
+	if (!ComponentClass->IsChildOf(UActorComponent::StaticClass()))
+	{
+		return FECACommandResult::Error(FString::Printf(TEXT("Class is not a component: %s"), *ComponentType));
+	}
+	
+	// Add the component
+	USCS_Node* NewNode = Blueprint->SimpleConstructionScript->CreateNode(ComponentClass, *ComponentName);
+	if (!NewNode)
+	{
+		return FECACommandResult::Error(TEXT("Failed to create component node"));
+	}
+	
+	// Handle attachment
+	FString AttachTo;
+	if (GetStringParam(Params, TEXT("attach_to"), AttachTo, false))
+	{
+		// Find the parent node
+		USCS_Node* ParentNode = nullptr;
+		TArray<USCS_Node*> AllNodes = Blueprint->SimpleConstructionScript->GetAllNodes();
+		for (USCS_Node* Node : AllNodes)
+		{
+			if (Node->GetVariableName().ToString() == AttachTo)
+			{
+				ParentNode = Node;
+				break;
+			}
+		}
+		
+		if (ParentNode)
+		{
+			ParentNode->AddChildNode(NewNode);
+		}
+		else
+		{
+			Blueprint->SimpleConstructionScript->AddNode(NewNode);
+		}
+	}
+	else
+	{
+		Blueprint->SimpleConstructionScript->AddNode(NewNode);
+	}
+	
+	FBlueprintEditorUtils::MarkBlueprintAsStructurallyModified(Blueprint);
+	
+	TSharedPtr<FJsonObject> Result = MakeResult();
+	Result->SetStringField(TEXT("component_name"), ComponentName);
+	Result->SetStringField(TEXT("component_class"), ComponentClass->GetName());
+	return FECACommandResult::Success(Result);
+}
+
+//------------------------------------------------------------------------------
+// CompileBlueprint
+//------------------------------------------------------------------------------
+
+FECACommandResult FECACommand_CompileBlueprint::Execute(const TSharedPtr<FJsonObject>& Params)
+{
+	FString BlueprintPath;
+	if (!GetStringParam(Params, TEXT("blueprint_path"), BlueprintPath))
+	{
+		return FECACommandResult::Error(TEXT("Missing required parameter: blueprint_path"));
+	}
+	
+	UBlueprint* Blueprint = LoadBlueprintByPath(BlueprintPath);
+	if (!Blueprint)
+	{
+		return FECACommandResult::Error(FString::Printf(TEXT("Blueprint not found: %s"), *BlueprintPath));
+	}
+	
+	FKismetEditorUtilities::CompileBlueprint(Blueprint);
+	
+	// Check compilation result
+	bool bHasErrors = Blueprint->Status == BS_Error;
+	
+	TSharedPtr<FJsonObject> Result = MakeResult();
+	Result->SetBoolField(TEXT("compiled"), true);
+	Result->SetBoolField(TEXT("has_errors"), bHasErrors);
+	Result->SetStringField(TEXT("status"), bHasErrors ? TEXT("Error") : TEXT("Success"));
+	return FECACommandResult::Success(Result);
+}
+
+//------------------------------------------------------------------------------
+// AddBlueprintVariable
+//------------------------------------------------------------------------------
+
+FECACommandResult FECACommand_AddBlueprintVariable::Execute(const TSharedPtr<FJsonObject>& Params)
+{
+	FString BlueprintPath;
+	if (!GetStringParam(Params, TEXT("blueprint_path"), BlueprintPath))
+	{
+		return FECACommandResult::Error(TEXT("Missing required parameter: blueprint_path"));
+	}
+	
+	FString VariableName;
+	if (!GetStringParam(Params, TEXT("variable_name"), VariableName))
+	{
+		return FECACommandResult::Error(TEXT("Missing required parameter: variable_name"));
+	}
+	
+	FString VariableType;
+	if (!GetStringParam(Params, TEXT("variable_type"), VariableType))
+	{
+		return FECACommandResult::Error(TEXT("Missing required parameter: variable_type"));
+	}
+	
+	UBlueprint* Blueprint = LoadBlueprintByPath(BlueprintPath);
+	if (!Blueprint)
+	{
+		return FECACommandResult::Error(FString::Printf(TEXT("Blueprint not found: %s"), *BlueprintPath));
+	}
+	
+	// Map variable type to pin type
+	FEdGraphPinType PinType;
+	
+	if (VariableType.Equals(TEXT("Boolean"), ESearchCase::IgnoreCase) || VariableType.Equals(TEXT("bool"), ESearchCase::IgnoreCase))
+	{
+		PinType.PinCategory = UEdGraphSchema_K2::PC_Boolean;
+	}
+	else if (VariableType.Equals(TEXT("Integer"), ESearchCase::IgnoreCase) || VariableType.Equals(TEXT("int"), ESearchCase::IgnoreCase))
+	{
+		PinType.PinCategory = UEdGraphSchema_K2::PC_Int;
+	}
+	else if (VariableType.Equals(TEXT("Float"), ESearchCase::IgnoreCase) || VariableType.Equals(TEXT("float"), ESearchCase::IgnoreCase) || VariableType.Equals(TEXT("double"), ESearchCase::IgnoreCase))
+	{
+		PinType.PinCategory = UEdGraphSchema_K2::PC_Real;
+		PinType.PinSubCategory = UEdGraphSchema_K2::PC_Double;
+	}
+	else if (VariableType.Equals(TEXT("String"), ESearchCase::IgnoreCase))
+	{
+		PinType.PinCategory = UEdGraphSchema_K2::PC_String;
+	}
+	else if (VariableType.Equals(TEXT("Name"), ESearchCase::IgnoreCase))
+	{
+		PinType.PinCategory = UEdGraphSchema_K2::PC_Name;
+	}
+	else if (VariableType.Equals(TEXT("Text"), ESearchCase::IgnoreCase))
+	{
+		PinType.PinCategory = UEdGraphSchema_K2::PC_Text;
+	}
+	else if (VariableType.Equals(TEXT("Vector"), ESearchCase::IgnoreCase))
+	{
+		PinType.PinCategory = UEdGraphSchema_K2::PC_Struct;
+		PinType.PinSubCategoryObject = TBaseStructure<FVector>::Get();
+	}
+	else if (VariableType.Equals(TEXT("Rotator"), ESearchCase::IgnoreCase))
+	{
+		PinType.PinCategory = UEdGraphSchema_K2::PC_Struct;
+		PinType.PinSubCategoryObject = TBaseStructure<FRotator>::Get();
+	}
+	else if (VariableType.Equals(TEXT("Transform"), ESearchCase::IgnoreCase))
+	{
+		PinType.PinCategory = UEdGraphSchema_K2::PC_Struct;
+		PinType.PinSubCategoryObject = TBaseStructure<FTransform>::Get();
+	}
+	else if (VariableType.Equals(TEXT("Color"), ESearchCase::IgnoreCase) || VariableType.Equals(TEXT("LinearColor"), ESearchCase::IgnoreCase))
+	{
+		PinType.PinCategory = UEdGraphSchema_K2::PC_Struct;
+		PinType.PinSubCategoryObject = TBaseStructure<FLinearColor>::Get();
+	}
+	else if (VariableType.Equals(TEXT("Actor"), ESearchCase::IgnoreCase))
+	{
+		PinType.PinCategory = UEdGraphSchema_K2::PC_Object;
+		PinType.PinSubCategoryObject = AActor::StaticClass();
+	}
+	else if (VariableType.Equals(TEXT("Object"), ESearchCase::IgnoreCase))
+	{
+		PinType.PinCategory = UEdGraphSchema_K2::PC_Object;
+		PinType.PinSubCategoryObject = UObject::StaticClass();
+	}
+	else if (VariableType.Equals(TEXT("Class"), ESearchCase::IgnoreCase))
+	{
+		PinType.PinCategory = UEdGraphSchema_K2::PC_Class;
+		PinType.PinSubCategoryObject = UObject::StaticClass();
+	}
+	else
+	{
+		// Try to find the class by name - supports specific class types like "DirectionalLight", "PointLight", etc.
+		UClass* FoundClass = nullptr;
+		
+		// First try common engine classes
+		TArray<FString> ClassPaths = {
+			FString::Printf(TEXT("/Script/Engine.%s"), *VariableType),
+			FString::Printf(TEXT("/Script/CoreUObject.%s"), *VariableType),
+			FString::Printf(TEXT("/Script/UMG.%s"), *VariableType),
+			VariableType  // Full path or just class name
+		};
+		
+		for (const FString& ClassPath : ClassPaths)
+		{
+			FoundClass = FindObject<UClass>(nullptr, *ClassPath);
+			if (FoundClass)
+			{
+				break;
+			}
+		}
+		
+		// Try loading it
+		if (!FoundClass)
+		{
+			FoundClass = LoadClass<UObject>(nullptr, *VariableType);
+		}
+		
+		if (FoundClass)
+		{
+			PinType.PinCategory = UEdGraphSchema_K2::PC_Object;
+			PinType.PinSubCategoryObject = FoundClass;
+		}
+		else
+		{
+			return FECACommandResult::Error(FString::Printf(TEXT("Unsupported or unknown variable type: %s. Supported: Boolean, Integer, Float, String, Vector, Rotator, Transform, Actor, Object, Class, or class names like DirectionalLight, PointLight, StaticMeshActor"), *VariableType));
+		}
+	}
+	
+	bool bSuccess = FBlueprintEditorUtils::AddMemberVariable(Blueprint, FName(*VariableName), PinType);
+	
+	if (!bSuccess)
+	{
+		return FECACommandResult::Error(TEXT("Failed to add variable"));
+	}
+	
+	// Handle additional options
+	bool bInstanceEditable = false;
+	bool bBlueprintReadOnly = false;
+	GetBoolParam(Params, TEXT("is_instance_editable"), bInstanceEditable, false);
+	GetBoolParam(Params, TEXT("is_blueprint_read_only"), bBlueprintReadOnly, false);
+	
+	FBlueprintEditorUtils::SetBlueprintOnlyEditableFlag(Blueprint, FName(*VariableName), !bInstanceEditable);
+	FBlueprintEditorUtils::SetBlueprintPropertyReadOnlyFlag(Blueprint, FName(*VariableName), bBlueprintReadOnly);
+	
+	// Handle default value if provided
+	FString DefaultValue;
+	bool bDefaultValueSet = false;
+	if (GetStringParam(Params, TEXT("default_value"), DefaultValue, false) && !DefaultValue.IsEmpty())
+	{
+		FName VarName(*VariableName);
+		
+		// Find the variable we just created
+		for (int32 i = 0; i < Blueprint->NewVariables.Num(); ++i)
+		{
+			if (Blueprint->NewVariables[i].VarName == VarName)
+			{
+				FBPVariableDescription& VarDesc = Blueprint->NewVariables[i];
+				VarDesc.DefaultValue = DefaultValue;
+				bDefaultValueSet = true;
+				
+				// For complex types, also update the CDO after recompile
+				// But for now, setting DefaultValue on the description should work for basic types
+				break;
+			}
+		}
+		
+		// Mark as structurally modified to ensure default value is applied
+		FBlueprintEditorUtils::MarkBlueprintAsStructurallyModified(Blueprint);
+	}
+	
+	TSharedPtr<FJsonObject> Result = MakeResult();
+	Result->SetStringField(TEXT("variable_name"), VariableName);
+	Result->SetStringField(TEXT("variable_type"), VariableType);
+	if (bDefaultValueSet)
+	{
+		Result->SetStringField(TEXT("default_value"), DefaultValue);
+	}
+	return FECACommandResult::Success(Result);
+}
+
+//------------------------------------------------------------------------------
+// SpawnBlueprintActor
+//------------------------------------------------------------------------------
+
+FECACommandResult FECACommand_SpawnBlueprintActor::Execute(const TSharedPtr<FJsonObject>& Params)
+{
+	FString BlueprintPath;
+	if (!GetStringParam(Params, TEXT("blueprint_path"), BlueprintPath))
+	{
+		return FECACommandResult::Error(TEXT("Missing required parameter: blueprint_path"));
+	}
+	
+	UBlueprint* Blueprint = LoadBlueprintByPath(BlueprintPath);
+	if (!Blueprint)
+	{
+		return FECACommandResult::Error(FString::Printf(TEXT("Blueprint not found: %s"), *BlueprintPath));
+	}
+	
+	if (!Blueprint->GeneratedClass)
+	{
+		return FECACommandResult::Error(TEXT("Blueprint has no generated class. Try compiling it first."));
+	}
+	
+	UWorld* World = GetEditorWorld();
+	if (!World)
+	{
+		return FECACommandResult::Error(TEXT("No world available"));
+	}
+	
+	FVector Location = FVector::ZeroVector;
+	FRotator Rotation = FRotator::ZeroRotator;
+	GetVectorParam(Params, TEXT("location"), Location, false);
+	GetRotatorParam(Params, TEXT("rotation"), Rotation, false);
+	
+	FActorSpawnParameters SpawnParams;
+	SpawnParams.SpawnCollisionHandlingOverride = ESpawnActorCollisionHandlingMethod::AdjustIfPossibleButAlwaysSpawn;
+	
+	AActor* NewActor = World->SpawnActor<AActor>(Blueprint->GeneratedClass, Location, Rotation, SpawnParams);
+	
+	if (!NewActor)
+	{
+		return FECACommandResult::Error(TEXT("Failed to spawn Blueprint actor"));
+	}
+	
+	FString ActorName;
+	if (GetStringParam(Params, TEXT("actor_name"), ActorName, false))
+	{
+		NewActor->SetActorLabel(*ActorName);
+	}
+	
+	TSharedPtr<FJsonObject> Result = MakeResult();
+	Result->SetStringField(TEXT("actor_name"), NewActor->GetActorLabel());
+	Result->SetStringField(TEXT("internal_name"), NewActor->GetName());
+	Result->SetObjectField(TEXT("transform"), TransformToJson(NewActor->GetTransform()));
+	
+	// Add spatial verification data - world bounds
+	FBox WorldBounds = NewActor->GetComponentsBoundingBox(true);
+	if (WorldBounds.IsValid)
+	{
+		TSharedPtr<FJsonObject> BoundsJson = MakeShared<FJsonObject>();
+		
+		TSharedPtr<FJsonObject> MinJson = MakeShared<FJsonObject>();
+		MinJson->SetNumberField(TEXT("x"), WorldBounds.Min.X);
+		MinJson->SetNumberField(TEXT("y"), WorldBounds.Min.Y);
+		MinJson->SetNumberField(TEXT("z"), WorldBounds.Min.Z);
+		BoundsJson->SetObjectField(TEXT("min"), MinJson);
+		
+		TSharedPtr<FJsonObject> MaxJson = MakeShared<FJsonObject>();
+		MaxJson->SetNumberField(TEXT("x"), WorldBounds.Max.X);
+		MaxJson->SetNumberField(TEXT("y"), WorldBounds.Max.Y);
+		MaxJson->SetNumberField(TEXT("z"), WorldBounds.Max.Z);
+		BoundsJson->SetObjectField(TEXT("max"), MaxJson);
+		
+		FVector Size = WorldBounds.GetSize();
+		TSharedPtr<FJsonObject> SizeJson = MakeShared<FJsonObject>();
+		SizeJson->SetNumberField(TEXT("x"), Size.X);
+		SizeJson->SetNumberField(TEXT("y"), Size.Y);
+		SizeJson->SetNumberField(TEXT("z"), Size.Z);
+		BoundsJson->SetObjectField(TEXT("size"), SizeJson);
+		
+		Result->SetObjectField(TEXT("world_bounds"), BoundsJson);
+	}
+	
+	return FECACommandResult::Success(Result);
+}
+
+//------------------------------------------------------------------------------
+// GetBlueprintInfo
+//------------------------------------------------------------------------------
+
+FECACommandResult FECACommand_GetBlueprintInfo::Execute(const TSharedPtr<FJsonObject>& Params)
+{
+	FString BlueprintPath;
+	if (!GetStringParam(Params, TEXT("blueprint_path"), BlueprintPath))
+	{
+		return FECACommandResult::Error(TEXT("Missing required parameter: blueprint_path"));
+	}
+	
+	UBlueprint* Blueprint = LoadBlueprintByPath(BlueprintPath);
+	if (!Blueprint)
+	{
+		return FECACommandResult::Error(FString::Printf(TEXT("Blueprint not found: %s"), *BlueprintPath));
+	}
+	
+	TSharedPtr<FJsonObject> Result = MakeResult();
+	Result->SetStringField(TEXT("name"), Blueprint->GetName());
+	Result->SetStringField(TEXT("path"), Blueprint->GetPathName());
+	Result->SetStringField(TEXT("parent_class"), Blueprint->ParentClass ? Blueprint->ParentClass->GetName() : TEXT("None"));
+	Result->SetStringField(TEXT("blueprint_type"), UEnum::GetValueAsString(Blueprint->BlueprintType));
+	
+	// Status
+	FString StatusStr;
+	switch (Blueprint->Status)
+	{
+		case BS_Unknown: StatusStr = TEXT("Unknown"); break;
+		case BS_Dirty: StatusStr = TEXT("Dirty"); break;
+		case BS_Error: StatusStr = TEXT("Error"); break;
+		case BS_UpToDate: StatusStr = TEXT("UpToDate"); break;
+		case BS_BeingCreated: StatusStr = TEXT("BeingCreated"); break;
+		default: StatusStr = TEXT("Unknown"); break;
+	}
+	Result->SetStringField(TEXT("status"), StatusStr);
+	
+	// Variables
+	TArray<TSharedPtr<FJsonValue>> VariablesArray;
+	for (const FBPVariableDescription& Var : Blueprint->NewVariables)
+	{
+		TSharedPtr<FJsonObject> VarObj = MakeShared<FJsonObject>();
+		VarObj->SetStringField(TEXT("name"), Var.VarName.ToString());
+		VarObj->SetStringField(TEXT("type"), Var.VarType.PinCategory.ToString());
+		VariablesArray.Add(MakeShared<FJsonValueObject>(VarObj));
+	}
+	Result->SetArrayField(TEXT("variables"), VariablesArray);
+	
+	// Components
+	TArray<TSharedPtr<FJsonValue>> ComponentsArray;
+	if (Blueprint->SimpleConstructionScript)
+	{
+		for (USCS_Node* Node : Blueprint->SimpleConstructionScript->GetAllNodes())
+		{
+			TSharedPtr<FJsonObject> CompObj = MakeShared<FJsonObject>();
+			CompObj->SetStringField(TEXT("name"), Node->GetVariableName().ToString());
+			CompObj->SetStringField(TEXT("class"), Node->ComponentClass ? Node->ComponentClass->GetName() : TEXT("None"));
+			ComponentsArray.Add(MakeShared<FJsonValueObject>(CompObj));
+		}
+	}
+	Result->SetArrayField(TEXT("components"), ComponentsArray);
+	
+	// Functions
+	TArray<TSharedPtr<FJsonValue>> FunctionsArray;
+	for (UEdGraph* Graph : Blueprint->FunctionGraphs)
+	{
+		TSharedPtr<FJsonObject> FuncObj = MakeShared<FJsonObject>();
+		FuncObj->SetStringField(TEXT("name"), Graph->GetName());
+		FunctionsArray.Add(MakeShared<FJsonValueObject>(FuncObj));
+	}
+	Result->SetArrayField(TEXT("functions"), FunctionsArray);
+	
+	return FECACommandResult::Success(Result);
+}
+
+//------------------------------------------------------------------------------
+// ListBlueprints
+//------------------------------------------------------------------------------
+
+FECACommandResult FECACommand_ListBlueprints::Execute(const TSharedPtr<FJsonObject>& Params)
+{
+	FString Path = TEXT("/Game/");
+	bool bRecursive = true;
+	GetStringParam(Params, TEXT("path"), Path, false);
+	GetBoolParam(Params, TEXT("recursive"), bRecursive, false);
+	
+	FAssetRegistryModule& AssetRegistryModule = FModuleManager::LoadModuleChecked<FAssetRegistryModule>("AssetRegistry");
+	IAssetRegistry& AssetRegistry = AssetRegistryModule.Get();
+	
+	TArray<FAssetData> AssetList;
+	AssetRegistry.GetAssetsByPath(*Path, AssetList, bRecursive);
+	
+	TArray<TSharedPtr<FJsonValue>> BlueprintsArray;
+	
+	for (const FAssetData& Asset : AssetList)
+	{
+		if (Asset.AssetClassPath == UBlueprint::StaticClass()->GetClassPathName())
+		{
+			TSharedPtr<FJsonObject> BPObj = MakeShared<FJsonObject>();
+			BPObj->SetStringField(TEXT("name"), Asset.AssetName.ToString());
+			BPObj->SetStringField(TEXT("path"), Asset.GetObjectPathString());
+			BlueprintsArray.Add(MakeShared<FJsonValueObject>(BPObj));
+		}
+	}
+	
+	TSharedPtr<FJsonObject> Result = MakeResult();
+	Result->SetArrayField(TEXT("blueprints"), BlueprintsArray);
+	Result->SetNumberField(TEXT("count"), BlueprintsArray.Num());
+	return FECACommandResult::Success(Result);
+}
+
+//------------------------------------------------------------------------------
+// OpenBlueprintEditor
+//------------------------------------------------------------------------------
+
+FECACommandResult FECACommand_OpenBlueprintEditor::Execute(const TSharedPtr<FJsonObject>& Params)
+{
+	FString BlueprintPath;
+	if (!GetStringParam(Params, TEXT("blueprint_path"), BlueprintPath))
+	{
+		return FECACommandResult::Error(TEXT("Missing required parameter: blueprint_path"));
+	}
+	
+	UBlueprint* Blueprint = LoadBlueprintByPath(BlueprintPath);
+	if (!Blueprint)
+	{
+		return FECACommandResult::Error(FString::Printf(TEXT("Blueprint not found: %s"), *BlueprintPath));
+	}
+	
+	GEditor->GetEditorSubsystem<UAssetEditorSubsystem>()->OpenEditorForAsset(Blueprint);
+	
+	return FECACommandResult::Success();
+}
+
+//------------------------------------------------------------------------------
+// DeleteBlueprintFunction
+//------------------------------------------------------------------------------
+
+FECACommandResult FECACommand_DeleteBlueprintFunction::Execute(const TSharedPtr<FJsonObject>& Params)
+{
+	FString BlueprintPath;
+	if (!GetStringParam(Params, TEXT("blueprint_path"), BlueprintPath))
+	{
+		return FECACommandResult::Error(TEXT("Missing required parameter: blueprint_path"));
+	}
+	
+	FString FunctionName;
+	if (!GetStringParam(Params, TEXT("function_name"), FunctionName))
+	{
+		return FECACommandResult::Error(TEXT("Missing required parameter: function_name"));
+	}
+	
+	UBlueprint* Blueprint = LoadBlueprintByPath(BlueprintPath);
+	if (!Blueprint)
+	{
+		return FECACommandResult::Error(FString::Printf(TEXT("Blueprint not found: %s"), *BlueprintPath));
+	}
+	
+	// Find the function graph
+	UEdGraph* FunctionGraph = nullptr;
+	for (UEdGraph* Graph : Blueprint->FunctionGraphs)
+	{
+		if (Graph && Graph->GetFName() == FName(*FunctionName))
+		{
+			FunctionGraph = Graph;
+			break;
+		}
+	}
+	
+	if (!FunctionGraph)
+	{
+		return FECACommandResult::Error(FString::Printf(TEXT("Function not found: %s"), *FunctionName));
+	}
+	
+	// Remove the function graph
+	FBlueprintEditorUtils::RemoveGraph(Blueprint, FunctionGraph, EGraphRemoveFlags::Recompile);
+	
+	// Mark the blueprint as modified
+	FBlueprintEditorUtils::MarkBlueprintAsStructurallyModified(Blueprint);
+	
+	TSharedPtr<FJsonObject> Result = MakeResult();
+	Result->SetStringField(TEXT("blueprint_path"), BlueprintPath);
+	Result->SetStringField(TEXT("deleted_function"), FunctionName);
+	Result->SetBoolField(TEXT("success"), true);
+	
+	return FECACommandResult::Success(Result);
+}
