ECABridge experimental MCP plugin

diff --git a/Engine/Plugins/Experimental/ECABridge/Source/ECABridge/Private/Commands/ECABlueprintLispCommands.cpp b/Engine/Plugins/Experimental/ECABridge/Source/ECABridge/Private/Commands/ECABlueprintLispCommands.cpp
new file mode 100644
--- /dev/null
+++ b/Engine/Plugins/Experimental/ECABridge/Source/ECABridge/Private/Commands/ECABlueprintLispCommands.cpp
@@ -0,0 +1,6647 @@
+// Copyright Epic Games, Inc. All Rights Reserved.
+
+#include "Commands/ECABlueprintLispCommands.h"
+#include "BlueprintLisp.h"
+#include "Serialization/JsonSerializer.h"
+#include "Engine/Blueprint.h"
+#include "EdGraph/EdGraph.h"
+#include "EdGraph/EdGraphNode.h"
+#include "EdGraph/EdGraphPin.h"
+#include "EdGraphSchema_K2.h"
+#include "K2Node_Event.h"
+#include "K2Node_CallFunction.h"
+#include "K2Node_IfThenElse.h"
+#include "K2Node_VariableGet.h"
+#include "K2Node_VariableSet.h"
+#include "K2Node_Self.h"
+#include "K2Node_CustomEvent.h"
+#include "K2Node_FunctionEntry.h"
+#include "K2Node_FunctionResult.h"
+#include "Kismet2/BlueprintEditorUtils.h"
+#include "K2Node_MacroInstance.h"
+#include "K2Node_DynamicCast.h"
+#include "K2Node_MacroInstance.h"
+#include "K2Node_ExecutionSequence.h"
+#include "K2Node_Switch.h"
+#include "K2Node_SwitchInteger.h"
+#include "K2Node_SwitchString.h"
+#include "K2Node_SwitchName.h"
+#include "K2Node_SwitchEnum.h"
+#include "K2Node_InputAction.h"
+#include "K2Node_InputKey.h"
+#include "K2Node_ComponentBoundEvent.h"
+#include "K2Node_ActorBoundEvent.h"
+#include "EdGraphNode_Comment.h"
+#include "Engine/SCS_Node.h"
+#include "Engine/SimpleConstructionScript.h"
+#include "Components/PrimitiveComponent.h"
+#include "Components/StaticMeshComponent.h"
+#include "Components/SkeletalMeshComponent.h"
+#include "Components/MeshComponent.h"
+#include "Components/LightComponent.h"
+#include "Components/AudioComponent.h"
+#include "Components/ShapeComponent.h"
+#include "Components/SceneComponent.h"
+#include "Components/ActorComponent.h"
+#include "BlueprintAutoLayout.h"
+#include "K2Node_MakeArray.h"
+#include "K2Node_GetArrayItem.h"
+#include "Kismet2/BlueprintEditorUtils.h"
+#include "Kismet/KismetSystemLibrary.h"
+#include "Kismet/KismetMathLibrary.h"
+#include "Kismet/KismetArrayLibrary.h"
+#include "Kismet/KismetStringLibrary.h"
+#include "Kismet/GameplayStatics.h"
+#include "GameFramework/Actor.h"
+#include "GameFramework/Pawn.h"
+#include "GameFramework/Character.h"
+#include "Kismet/BlueprintFunctionLibrary.h"
+#include "UObject/UObjectIterator.h"
+
+//------------------------------------------------------------------------------
+// Helper to add Geometry Script classes to function search
+//------------------------------------------------------------------------------
+static void AddGeometryScriptClasses(TArray<UClass*>& ClassesToSearch)
+{
+	// Geometry Script library classes for procedural mesh operations
+	static const TCHAR* GeometryScriptClassPaths[] = {
+		TEXT("/Script/GeometryScriptingCore.GeometryScriptLibrary_MeshPrimitiveFunctions"),
+		TEXT("/Script/GeometryScriptingCore.GeometryScriptLibrary_MeshBasicEditFunctions"),
+		TEXT("/Script/GeometryScriptingCore.GeometryScriptLibrary_MeshTransformFunctions"),
+		TEXT("/Script/GeometryScriptingCore.GeometryScriptLibrary_MeshBooleanFunctions"),
+		TEXT("/Script/GeometryScriptingCore.GeometryScriptLibrary_MeshModelingFunctions"),
+		TEXT("/Script/GeometryScriptingCore.GeometryScriptLibrary_MeshNormalsFunctions"),
+		TEXT("/Script/GeometryScriptingCore.GeometryScriptLibrary_MeshUVFunctions"),
+		TEXT("/Script/GeometryScriptingCore.GeometryScriptLibrary_MeshMaterialFunctions"),
+		TEXT("/Script/GeometryScriptingCore.GeometryScriptLibrary_MeshDeformFunctions"),
+		TEXT("/Script/GeometryScriptingCore.GeometryScriptLibrary_MeshSubdivideFunctions"),
+		TEXT("/Script/GeometryScriptingCore.GeometryScriptLibrary_MeshSimplifyFunctions"),
+		TEXT("/Script/GeometryScriptingCore.GeometryScriptLibrary_MeshRepairFunctions"),
+		TEXT("/Script/GeometryScriptingCore.GeometryScriptLibrary_MeshQueryFunctions"),
+		TEXT("/Script/GeometryScriptingCore.GeometryScriptLibrary_MeshSpatialFunctions"),
+		TEXT("/Script/GeometryScriptingCore.GeometryScriptLibrary_MeshSelectionFunctions"),
+		TEXT("/Script/GeometryScriptingCore.GeometryScriptLibrary_ContainmentFunctions"),
+		TEXT("/Script/GeometryScriptingCore.GeometryScriptLibrary_CollisionFunctions"),
+		TEXT("/Script/GeometryScriptingCore.GeometryScriptLibrary_MeshVoxelFunctions"),
+		TEXT("/Script/GeometryScriptingCore.GeometryScriptLibrary_PolyPathFunctions"),
+		TEXT("/Script/GeometryScriptingCore.GeometryScriptLibrary_ListUtilityFunctions"),
+		TEXT("/Script/GeometryScriptingEditor.GeometryScriptLibrary_CreateNewAssetUtilityFunctions"),
+	};
+	
+	for (const TCHAR* ClassPath : GeometryScriptClassPaths)
+	{
+		UClass* GeoClass = FindObject<UClass>(nullptr, ClassPath);
+		if (!GeoClass)
+		{
+			GeoClass = LoadObject<UClass>(nullptr, ClassPath);
+		}
+		if (!GeoClass)
+		{
+			GeoClass = StaticLoadClass(UObject::StaticClass(), nullptr, ClassPath);
+		}
+		if (GeoClass)
+		{
+			ClassesToSearch.Add(GeoClass);
+		}
+	}
+}
+
+//------------------------------------------------------------------------------
+// Helper to set numeric default value on a pin with proper type handling
+//------------------------------------------------------------------------------
+static void SetNumericPinDefaultValue(UEdGraphPin* Pin, double Value)
+{
+	if (!Pin) return;
+	
+	// Check the pin's expected type and format accordingly
+	if (Pin->PinType.PinCategory == UEdGraphSchema_K2::PC_Int ||
+	    Pin->PinType.PinCategory == UEdGraphSchema_K2::PC_Int64)
+	{
+		// Integer pin - use integer format (no decimal point)
+		Pin->DefaultValue = FString::FromInt(FMath::RoundToInt(Value));
+	}
+	else if (Pin->PinType.PinCategory == UEdGraphSchema_K2::PC_Byte)
+	{
+		// Byte pin - clamp and use integer format
+		Pin->DefaultValue = FString::FromInt(FMath::Clamp(FMath::RoundToInt(Value), 0, 255));
+	}
+	else
+	{
+		// Float/Double/Real or unknown - use float format
+		Pin->DefaultValue = FString::SanitizeFloat(Value);
+	}
+}
+
+// Register commands
+REGISTER_ECA_COMMAND(FECACommand_ParseBlueprintLisp)
+REGISTER_ECA_COMMAND(FECACommand_BlueprintToLisp)
+REGISTER_ECA_COMMAND(FECACommand_LispToBlueprint)
+REGISTER_ECA_COMMAND(FECACommand_BlueprintLispHelp)
+
+//------------------------------------------------------------------------------
+// Helper to convert FLispNode to JSON for AST inspection
+//------------------------------------------------------------------------------
+static TSharedPtr<FJsonObject> LispNodeToJson(const FLispNodePtr& Node)
+{
+	TSharedPtr<FJsonObject> Json = MakeShared<FJsonObject>();
+	
+	if (!Node.IsValid())
+	{
+		Json->SetStringField(TEXT("type"), TEXT("nil"));
+		return Json;
+	}
+	
+	switch (Node->Type)
+	{
+	case ELispNodeType::Nil:
+		Json->SetStringField(TEXT("type"), TEXT("nil"));
+		break;
+		
+	case ELispNodeType::Symbol:
+		Json->SetStringField(TEXT("type"), TEXT("symbol"));
+		Json->SetStringField(TEXT("value"), Node->StringValue);
+		break;
+		
+	case ELispNodeType::Keyword:
+		Json->SetStringField(TEXT("type"), TEXT("keyword"));
+		Json->SetStringField(TEXT("value"), Node->StringValue);
+		break;
+		
+	case ELispNodeType::Number:
+		Json->SetStringField(TEXT("type"), TEXT("number"));
+		Json->SetNumberField(TEXT("value"), Node->NumberValue);
+		break;
+		
+	case ELispNodeType::String:
+		Json->SetStringField(TEXT("type"), TEXT("string"));
+		Json->SetStringField(TEXT("value"), Node->StringValue);
+		break;
+		
+	case ELispNodeType::List:
+		{
+			Json->SetStringField(TEXT("type"), TEXT("list"));
+			TArray<TSharedPtr<FJsonValue>> ChildrenArray;
+			for (const auto& Child : Node->Children)
+			{
+				ChildrenArray.Add(MakeShared<FJsonValueObject>(LispNodeToJson(Child)));
+			}
+			Json->SetArrayField(TEXT("children"), ChildrenArray);
+			
+			FString FormName = Node->GetFormName();
+			if (!FormName.IsEmpty())
+			{
+				Json->SetStringField(TEXT("form"), FormName);
+			}
+		}
+		break;
+	}
+	
+	if (Node->Line > 0)
+	{
+		Json->SetNumberField(TEXT("line"), Node->Line);
+		Json->SetNumberField(TEXT("column"), Node->Column);
+	}
+	
+	return Json;
+}
+
+//------------------------------------------------------------------------------
+// Blueprint -> Lisp Conversion Helpers
+//------------------------------------------------------------------------------
+
+static UBlueprint* LoadBlueprintByPathLisp(const FString& Path)
+{
+	return LoadObject<UBlueprint>(nullptr, *Path);
+}
+
+static UEdGraph* FindGraphByNameLisp(UBlueprint* Blueprint, const FString& GraphName)
+{
+	for (UEdGraph* Graph : Blueprint->UbergraphPages)
+	{
+		if (Graph->GetName() == GraphName || GraphName.Equals(TEXT("EventGraph"), ESearchCase::IgnoreCase))
+		{
+			return Graph;
+		}
+	}
+	for (UEdGraph* Graph : Blueprint->FunctionGraphs)
+	{
+		if (Graph->GetName() == GraphName)
+		{
+			return Graph;
+		}
+	}
+	return nullptr;
+}
+
+static FString GetCleanNodeName(UEdGraphNode* Node)
+{
+	if (UK2Node_CallFunction* CallNode = Cast<UK2Node_CallFunction>(Node))
+	{
+		UFunction* Func = CallNode->GetTargetFunction();
+		if (Func)
+		{
+			return Func->GetName();
+		}
+	}
+	
+	FString Title = Node->GetNodeTitle(ENodeTitleType::ListView).ToString();
+	Title.ReplaceInline(TEXT(" "), TEXT(""));
+	return Title;
+}
+
+static UEdGraphPin* GetThenPin(UEdGraphNode* Node)
+{
+	for (UEdGraphPin* Pin : Node->Pins)
+	{
+		if (Pin->Direction == EGPD_Output && 
+			Pin->PinType.PinCategory == UEdGraphSchema_K2::PC_Exec)
+		{
+			return Pin;
+		}
+	}
+	return nullptr;
+}
+
+static FString PinTypeToLispType(const FEdGraphPinType& PinType)
+{
+	FString Category = PinType.PinCategory.ToString();
+	
+	if (Category == TEXT("bool")) return TEXT("Boolean");
+	if (Category == TEXT("int")) return TEXT("Integer");
+	if (Category == TEXT("float")) return TEXT("Float");
+	if (Category == TEXT("double")) return TEXT("Double");
+	if (Category == TEXT("name")) return TEXT("Name");
+	if (Category == TEXT("string")) return TEXT("String");
+	if (Category == TEXT("text")) return TEXT("Text");
+	if (Category == TEXT("struct") || Category == TEXT("object") || Category == TEXT("class"))
+	{
+		if (PinType.PinSubCategoryObject.IsValid())
+		{
+			return PinType.PinSubCategoryObject->GetName();
+		}
+	}
+	
+	return Category;
+}
+
+// Forward declarations
+static FLispNodePtr ConvertPureExpressionToLisp(UEdGraphPin* ValuePin, UEdGraph* Graph, TSet<UEdGraphNode*>& VisitedNodes);
+static FLispNodePtr ConvertExecChainToLisp(UEdGraphPin* ExecPin, UEdGraph* Graph, TSet<UEdGraphNode*>& VisitedNodes, bool bIncludePositions);
+static FLispNodePtr ConvertNodeToLisp(UEdGraphNode* Node, UEdGraph* Graph, TSet<UEdGraphNode*>& VisitedNodes, bool bIncludePositions);
+
+static FLispNodePtr ConvertPureExpressionToLisp(UEdGraphPin* ValuePin, UEdGraph* Graph, TSet<UEdGraphNode*>& VisitedNodes)
+{
+	if (!ValuePin || ValuePin->LinkedTo.Num() == 0)
+	{
+		if (!ValuePin)
+		{
+			return FLispNode::MakeSymbol(TEXT("nil"));
+		}
+		
+		// Check for object reference first (materials, textures, meshes, etc.)
+		if (ValuePin->DefaultObject)
+		{
+			// Return as (asset "path/to/asset")
+			TArray<FLispNodePtr> AssetArgs;
+			AssetArgs.Add(FLispNode::MakeSymbol(TEXT("asset")));
+			AssetArgs.Add(FLispNode::MakeString(ValuePin->DefaultObject->GetPathName()));
+			return FLispNode::MakeList(AssetArgs);
+		}
+		
+		FString DefaultValue = ValuePin->DefaultValue;
+		if (DefaultValue.IsEmpty())
+		{
+			if (ValuePin->PinType.PinCategory == UEdGraphSchema_K2::PC_Boolean)
+			{
+				return FLispNode::MakeSymbol(TEXT("false"));
+			}
+			if (ValuePin->PinType.PinCategory == UEdGraphSchema_K2::PC_Int ||
+				ValuePin->PinType.PinCategory == UEdGraphSchema_K2::PC_Float)
+			{
+				return FLispNode::MakeNumber(0);
+			}
+			// Check for object/class type with no default
+			if (ValuePin->PinType.PinCategory == UEdGraphSchema_K2::PC_Object ||
+				ValuePin->PinType.PinCategory == UEdGraphSchema_K2::PC_Class ||
+				ValuePin->PinType.PinCategory == UEdGraphSchema_K2::PC_SoftObject ||
+				ValuePin->PinType.PinCategory == UEdGraphSchema_K2::PC_SoftClass)
+			{
+				return FLispNode::MakeSymbol(TEXT("nil"));
+			}
+			return FLispNode::MakeSymbol(TEXT("nil"));
+		}
+		
+		if (ValuePin->PinType.PinCategory == UEdGraphSchema_K2::PC_Boolean)
+		{
+			return FLispNode::MakeSymbol(DefaultValue.ToBool() ? TEXT("true") : TEXT("false"));
+		}
+		if (ValuePin->PinType.PinCategory == UEdGraphSchema_K2::PC_Int ||
+			ValuePin->PinType.PinCategory == UEdGraphSchema_K2::PC_Float)
+		{
+			return FLispNode::MakeNumber(FCString::Atod(*DefaultValue));
+		}
+		if (ValuePin->PinType.PinCategory == UEdGraphSchema_K2::PC_String ||
+			ValuePin->PinType.PinCategory == UEdGraphSchema_K2::PC_Name ||
+			ValuePin->PinType.PinCategory == UEdGraphSchema_K2::PC_Text)
+		{
+			return FLispNode::MakeString(DefaultValue);
+		}
+		
+		// For object types stored as string paths (soft references)
+		if ((ValuePin->PinType.PinCategory == UEdGraphSchema_K2::PC_Object ||
+			 ValuePin->PinType.PinCategory == UEdGraphSchema_K2::PC_SoftObject) &&
+			DefaultValue.Contains(TEXT("/")))
+		{
+			TArray<FLispNodePtr> AssetArgs;
+			AssetArgs.Add(FLispNode::MakeSymbol(TEXT("asset")));
+			AssetArgs.Add(FLispNode::MakeString(DefaultValue));
+			return FLispNode::MakeList(AssetArgs);
+		}
+		
+		return FLispNode::MakeString(DefaultValue);
+	}
+	
+	UEdGraphPin* SourcePin = ValuePin->LinkedTo[0];
+	UEdGraphNode* SourceNode = SourcePin->GetOwningNode();
+	
+	if (UK2Node_VariableGet* VarGet = Cast<UK2Node_VariableGet>(SourceNode))
+	{
+		return FLispNode::MakeSymbol(VarGet->VariableReference.GetMemberName().ToString());
+	}
+	
+	if (UK2Node_Self* SelfNode = Cast<UK2Node_Self>(SourceNode))
+	{
+		return FLispNode::MakeSymbol(TEXT("self"));
+	}
+	
+	// Handle event output pins (component bound events, custom events, etc.)
+	// Return the pin name as a symbol (e.g., OtherActor, DeltaSeconds)
+	if (UK2Node_Event* EventNode = Cast<UK2Node_Event>(SourceNode))
+	{
+		return FLispNode::MakeSymbol(SourcePin->PinName.ToString());
+	}
+	
+	if (UK2Node_ComponentBoundEvent* CompEventNode = Cast<UK2Node_ComponentBoundEvent>(SourceNode))
+	{
+		return FLispNode::MakeSymbol(SourcePin->PinName.ToString());
+	}
+	
+	if (UK2Node_CustomEvent* CustomEventNode = Cast<UK2Node_CustomEvent>(SourceNode))
+	{
+		return FLispNode::MakeSymbol(SourcePin->PinName.ToString());
+	}
+	
+	if (UK2Node_CallFunction* CallNode = Cast<UK2Node_CallFunction>(SourceNode))
+	{
+		UFunction* Func = CallNode->GetTargetFunction();
+		FString FuncName = Func ? Func->GetName() : GetCleanNodeName(SourceNode);
+		
+		// Special handling for BreakVector/BreakRotator - convert to (. struct Field) syntax
+		if (FuncName == TEXT("BreakVector") || FuncName == TEXT("BreakRotator") || 
+			FuncName == TEXT("BreakVector2D") || FuncName == TEXT("BreakTransform") ||
+			FuncName == TEXT("BreakColor"))
+		{
+			// Find which output pin we're using (X, Y, Z, etc.)
+			FString FieldName = SourcePin->PinName.ToString();
+			
+			// Find the struct input pin and convert it
+			for (UEdGraphPin* Pin : CallNode->Pins)
+			{
+				if (Pin->Direction == EGPD_Input && 
+					Pin->PinType.PinCategory == UEdGraphSchema_K2::PC_Struct)
+				{
+					TArray<FLispNodePtr> DotArgs;
+					DotArgs.Add(FLispNode::MakeSymbol(TEXT(".")));
+					DotArgs.Add(ConvertPureExpressionToLisp(Pin, Graph, VisitedNodes));
+					DotArgs.Add(FLispNode::MakeSymbol(FieldName));
+					return FLispNode::MakeList(DotArgs);
+				}
+			}
+			
+			// Fallback if no struct input found - try first non-exec input
+			for (UEdGraphPin* Pin : CallNode->Pins)
+			{
+				if (Pin->Direction == EGPD_Input && 
+					Pin->PinType.PinCategory != UEdGraphSchema_K2::PC_Exec &&
+					!Pin->bHidden)
+				{
+					TArray<FLispNodePtr> DotArgs;
+					DotArgs.Add(FLispNode::MakeSymbol(TEXT(".")));
+					DotArgs.Add(ConvertPureExpressionToLisp(Pin, Graph, VisitedNodes));
+					DotArgs.Add(FLispNode::MakeSymbol(FieldName));
+					return FLispNode::MakeList(DotArgs);
+				}
+			}
+		}
+		
+		// Map common math functions to Lisp operators
+		static TMap<FString, FString> MathOpMap;
+		if (MathOpMap.Num() == 0)
+		{
+			// Addition
+			MathOpMap.Add(TEXT("Add_FloatFloat"), TEXT("+"));
+			MathOpMap.Add(TEXT("Add_DoubleDouble"), TEXT("+"));
+			MathOpMap.Add(TEXT("Add_IntInt"), TEXT("+"));
+			MathOpMap.Add(TEXT("Add_Int64Int64"), TEXT("+"));
+			MathOpMap.Add(TEXT("Add_VectorVector"), TEXT("+"));
+			MathOpMap.Add(TEXT("Add_Vector2DVector2D"), TEXT("+"));
+			
+			// Subtraction
+			MathOpMap.Add(TEXT("Subtract_FloatFloat"), TEXT("-"));
+			MathOpMap.Add(TEXT("Subtract_DoubleDouble"), TEXT("-"));
+			MathOpMap.Add(TEXT("Subtract_IntInt"), TEXT("-"));
+			MathOpMap.Add(TEXT("Subtract_Int64Int64"), TEXT("-"));
+			MathOpMap.Add(TEXT("Subtract_VectorVector"), TEXT("-"));
+			MathOpMap.Add(TEXT("Subtract_Vector2DVector2D"), TEXT("-"));
+			
+			// Multiplication
+			MathOpMap.Add(TEXT("Multiply_FloatFloat"), TEXT("*"));
+			MathOpMap.Add(TEXT("Multiply_DoubleDouble"), TEXT("*"));
+			MathOpMap.Add(TEXT("Multiply_IntInt"), TEXT("*"));
+			MathOpMap.Add(TEXT("Multiply_Int64Int64"), TEXT("*"));
+			MathOpMap.Add(TEXT("Multiply_VectorFloat"), TEXT("*"));
+			MathOpMap.Add(TEXT("Multiply_VectorInt"), TEXT("*"));
+			MathOpMap.Add(TEXT("Multiply_VectorVector"), TEXT("*"));
+			MathOpMap.Add(TEXT("Multiply_Vector2DFloat"), TEXT("*"));
+			MathOpMap.Add(TEXT("Multiply_Vector2DVector2D"), TEXT("*"));
+			
+			// Division
+			MathOpMap.Add(TEXT("Divide_FloatFloat"), TEXT("/"));
+			MathOpMap.Add(TEXT("Divide_DoubleDouble"), TEXT("/"));
+			MathOpMap.Add(TEXT("Divide_IntInt"), TEXT("/"));
+			MathOpMap.Add(TEXT("Divide_Int64Int64"), TEXT("/"));
+			MathOpMap.Add(TEXT("Divide_VectorFloat"), TEXT("/"));
+			MathOpMap.Add(TEXT("Divide_VectorInt"), TEXT("/"));
+			MathOpMap.Add(TEXT("Divide_VectorVector"), TEXT("/"));
+			MathOpMap.Add(TEXT("Divide_Vector2DFloat"), TEXT("/"));
+			
+			// Modulo
+			MathOpMap.Add(TEXT("Percent_FloatFloat"), TEXT("%"));
+			MathOpMap.Add(TEXT("Percent_IntInt"), TEXT("%"));
+			MathOpMap.Add(TEXT("Percent_Int64Int64"), TEXT("%"));
+			
+			// Comparisons - Float variants
+			MathOpMap.Add(TEXT("Less_FloatFloat"), TEXT("<"));
+			MathOpMap.Add(TEXT("Less_DoubleDouble"), TEXT("<"));
+			MathOpMap.Add(TEXT("Less_IntInt"), TEXT("<"));
+			MathOpMap.Add(TEXT("Less_Int64Int64"), TEXT("<"));
+			MathOpMap.Add(TEXT("Greater_FloatFloat"), TEXT(">"));
+			MathOpMap.Add(TEXT("Greater_DoubleDouble"), TEXT(">"));
+			MathOpMap.Add(TEXT("Greater_IntInt"), TEXT(">"));
+			MathOpMap.Add(TEXT("Greater_Int64Int64"), TEXT(">"));
+			MathOpMap.Add(TEXT("LessEqual_FloatFloat"), TEXT("<="));
+			MathOpMap.Add(TEXT("LessEqual_DoubleDouble"), TEXT("<="));
+			MathOpMap.Add(TEXT("LessEqual_IntInt"), TEXT("<="));
+			MathOpMap.Add(TEXT("LessEqual_Int64Int64"), TEXT("<="));
+			MathOpMap.Add(TEXT("GreaterEqual_FloatFloat"), TEXT(">="));
+			MathOpMap.Add(TEXT("GreaterEqual_DoubleDouble"), TEXT(">="));
+			MathOpMap.Add(TEXT("GreaterEqual_IntInt"), TEXT(">="));
+			MathOpMap.Add(TEXT("GreaterEqual_Int64Int64"), TEXT(">="));
+			
+			// Equality
+			MathOpMap.Add(TEXT("EqualEqual_FloatFloat"), TEXT("=="));
+			MathOpMap.Add(TEXT("EqualEqual_DoubleDouble"), TEXT("=="));
+			MathOpMap.Add(TEXT("EqualEqual_IntInt"), TEXT("=="));
+			MathOpMap.Add(TEXT("EqualEqual_Int64Int64"), TEXT("=="));
+			MathOpMap.Add(TEXT("EqualEqual_BoolBool"), TEXT("=="));
+			MathOpMap.Add(TEXT("EqualEqual_NameName"), TEXT("=="));
+			MathOpMap.Add(TEXT("EqualEqual_StringString"), TEXT("=="));
+			MathOpMap.Add(TEXT("EqualEqual_ObjectObject"), TEXT("=="));
+			MathOpMap.Add(TEXT("EqualEqual_ClassClass"), TEXT("=="));
+			MathOpMap.Add(TEXT("NotEqual_FloatFloat"), TEXT("!="));
+			MathOpMap.Add(TEXT("NotEqual_DoubleDouble"), TEXT("!="));
+			MathOpMap.Add(TEXT("NotEqual_IntInt"), TEXT("!="));
+			MathOpMap.Add(TEXT("NotEqual_Int64Int64"), TEXT("!="));
+			MathOpMap.Add(TEXT("NotEqual_BoolBool"), TEXT("!="));
+			MathOpMap.Add(TEXT("NotEqual_NameName"), TEXT("!="));
+			MathOpMap.Add(TEXT("NotEqual_StringString"), TEXT("!="));
+			MathOpMap.Add(TEXT("NotEqual_ObjectObject"), TEXT("!="));
+			
+			// Boolean operations
+			MathOpMap.Add(TEXT("BooleanAND"), TEXT("and"));
+			MathOpMap.Add(TEXT("BooleanOR"), TEXT("or"));
+			MathOpMap.Add(TEXT("Not_PreBool"), TEXT("not"));
+			MathOpMap.Add(TEXT("BooleanNAND"), TEXT("nand"));
+			MathOpMap.Add(TEXT("BooleanNOR"), TEXT("nor"));
+			MathOpMap.Add(TEXT("BooleanXOR"), TEXT("xor"));
+		}
+		
+		// Check if this is a math operation
+		if (FString* MathOp = MathOpMap.Find(FuncName))
+		{
+			TArray<FLispNodePtr> MathArgs;
+			MathArgs.Add(FLispNode::MakeSymbol(*MathOp));
+			
+			for (UEdGraphPin* Pin : CallNode->Pins)
+			{
+				if (Pin->Direction == EGPD_Input && 
+					Pin->PinType.PinCategory != UEdGraphSchema_K2::PC_Exec &&
+					!Pin->bHidden)
+				{
+					MathArgs.Add(ConvertPureExpressionToLisp(Pin, Graph, VisitedNodes));
+				}
+			}
+			return FLispNode::MakeList(MathArgs);
+		}
+		
+		// Check for common pure functions with simpler names
+		static TMap<FString, FString> FuncNameMap;
+		if (FuncNameMap.Num() == 0)
+		{
+			// Vector/Rotator/Transform construction
+			FuncNameMap.Add(TEXT("MakeVector"), TEXT("vec"));
+			FuncNameMap.Add(TEXT("MakeVector2D"), TEXT("vec2"));
+			FuncNameMap.Add(TEXT("MakeRotator"), TEXT("rot"));
+			FuncNameMap.Add(TEXT("MakeColor"), TEXT("color"));
+			FuncNameMap.Add(TEXT("MakeTransform"), TEXT("transform"));
+			// Note: BreakVector/BreakRotator are handled specially below to generate (. struct Field) syntax
+			
+			// Type conversions
+			FuncNameMap.Add(TEXT("Conv_IntToFloat"), TEXT("int->float"));
+			FuncNameMap.Add(TEXT("Conv_FloatToInt"), TEXT("float->int"));
+			FuncNameMap.Add(TEXT("Conv_IntToString"), TEXT("int->string"));
+			FuncNameMap.Add(TEXT("Conv_FloatToString"), TEXT("float->string"));
+			FuncNameMap.Add(TEXT("Conv_BoolToString"), TEXT("bool->string"));
+			FuncNameMap.Add(TEXT("Conv_StringToInt"), TEXT("string->int"));
+			FuncNameMap.Add(TEXT("Conv_StringToFloat"), TEXT("string->float"));
+			
+			// Common checks
+			FuncNameMap.Add(TEXT("IsValid"), TEXT("valid?"));
+			FuncNameMap.Add(TEXT("IsValidClass"), TEXT("valid-class?"));
+			
+			// String operations
+			FuncNameMap.Add(TEXT("Concat_StrStr"), TEXT("str+"));
+			FuncNameMap.Add(TEXT("Len"), TEXT("str-len"));
+			FuncNameMap.Add(TEXT("Contains"), TEXT("str-contains?"));
+			
+			// Math utilities
+			FuncNameMap.Add(TEXT("Abs"), TEXT("abs"));
+			FuncNameMap.Add(TEXT("Sin"), TEXT("sin"));
+			FuncNameMap.Add(TEXT("Cos"), TEXT("cos"));
+			FuncNameMap.Add(TEXT("Tan"), TEXT("tan"));
+			FuncNameMap.Add(TEXT("Sqrt"), TEXT("sqrt"));
+			FuncNameMap.Add(TEXT("Square"), TEXT("sqr"));
+			FuncNameMap.Add(TEXT("FMin"), TEXT("min"));
+			FuncNameMap.Add(TEXT("FMax"), TEXT("max"));
+			FuncNameMap.Add(TEXT("Clamp"), TEXT("clamp"));
+			FuncNameMap.Add(TEXT("Lerp"), TEXT("lerp"));
+			FuncNameMap.Add(TEXT("FInterpTo"), TEXT("interp"));
+			FuncNameMap.Add(TEXT("RandomFloat"), TEXT("random"));
+			FuncNameMap.Add(TEXT("RandomFloatInRange"), TEXT("random-range"));
+			FuncNameMap.Add(TEXT("RandomInteger"), TEXT("random-int"));
+			FuncNameMap.Add(TEXT("RandomIntegerInRange"), TEXT("random-int-range"));
+			
+			// Vector operations
+			FuncNameMap.Add(TEXT("VSize"), TEXT("vec-length"));
+			FuncNameMap.Add(TEXT("Normal"), TEXT("vec-normalize"));
+			FuncNameMap.Add(TEXT("Dot_VectorVector"), TEXT("dot"));
+			FuncNameMap.Add(TEXT("Cross_VectorVector"), TEXT("cross"));
+			FuncNameMap.Add(TEXT("VInterpTo"), TEXT("vec-interp"));
+			
+			// Actor/Component
+			FuncNameMap.Add(TEXT("GetActorLocation"), TEXT("get-location"));
+			FuncNameMap.Add(TEXT("GetActorRotation"), TEXT("get-rotation"));
+			FuncNameMap.Add(TEXT("GetActorScale3D"), TEXT("get-scale"));
+			FuncNameMap.Add(TEXT("K2_GetActorLocation"), TEXT("get-location"));
+			FuncNameMap.Add(TEXT("K2_GetActorRotation"), TEXT("get-rotation"));
+			FuncNameMap.Add(TEXT("GetVelocity"), TEXT("get-velocity"));
+			FuncNameMap.Add(TEXT("GetActorForwardVector"), TEXT("get-forward"));
+			FuncNameMap.Add(TEXT("GetActorRightVector"), TEXT("get-right"));
+			FuncNameMap.Add(TEXT("GetActorUpVector"), TEXT("get-up"));
+		}
+		
+		FString DisplayName = FuncName;
+		if (FString* MappedName = FuncNameMap.Find(FuncName))
+		{
+			DisplayName = *MappedName;
+		}
+		
+		UEdGraphPin* TargetPin = CallNode->FindPin(UEdGraphSchema_K2::PN_Self);
+		if (TargetPin && TargetPin->LinkedTo.Num() > 0)
+		{
+			TArray<FLispNodePtr> CallArgs;
+			CallArgs.Add(FLispNode::MakeSymbol(TEXT("call")));
+			CallArgs.Add(ConvertPureExpressionToLisp(TargetPin, Graph, VisitedNodes));
+			CallArgs.Add(FLispNode::MakeSymbol(DisplayName));
+			
+			for (UEdGraphPin* Pin : CallNode->Pins)
+			{
+				if (Pin->Direction == EGPD_Input && 
+					Pin->PinType.PinCategory != UEdGraphSchema_K2::PC_Exec &&
+					Pin != TargetPin && !Pin->bHidden)
+				{
+					CallArgs.Add(ConvertPureExpressionToLisp(Pin, Graph, VisitedNodes));
+				}
+			}
+			return FLispNode::MakeList(CallArgs);
+		}
+		
+		TArray<FLispNodePtr> Args;
+		Args.Add(FLispNode::MakeSymbol(DisplayName));
+		for (UEdGraphPin* Pin : CallNode->Pins)
+		{
+			if (Pin->Direction == EGPD_Input && 
+				Pin->PinType.PinCategory != UEdGraphSchema_K2::PC_Exec &&
+				!Pin->bHidden)
+			{
+				Args.Add(ConvertPureExpressionToLisp(Pin, Graph, VisitedNodes));
+			}
+		}
+		return FLispNode::MakeList(Args);
+	}
+	
+	return FLispNode::MakeSymbol(GetCleanNodeName(SourceNode));
+}
+
+static FLispNodePtr ConvertNodeToLisp(UEdGraphNode* Node, UEdGraph* Graph, TSet<UEdGraphNode*>& VisitedNodes, bool bIncludePositions)
+{
+	if (!Node)
+	{
+		return FLispNode::MakeNil();
+	}
+	
+	// Handle Branch
+	if (UK2Node_IfThenElse* BranchNode = Cast<UK2Node_IfThenElse>(Node))
+	{
+		TArray<FLispNodePtr> BranchArgs;
+		BranchArgs.Add(FLispNode::MakeSymbol(TEXT("branch")));
+		
+		UEdGraphPin* ConditionPin = BranchNode->GetConditionPin();
+		BranchArgs.Add(ConvertPureExpressionToLisp(ConditionPin, Graph, VisitedNodes));
+		
+		BranchArgs.Add(FLispNode::MakeKeyword(TEXT(":true")));
+		UEdGraphPin* TruePin = BranchNode->GetThenPin();
+		FLispNodePtr TrueBranch = ConvertExecChainToLisp(TruePin, Graph, VisitedNodes, bIncludePositions);
+		BranchArgs.Add(TrueBranch.IsValid() && !TrueBranch->IsNil() ? TrueBranch : FLispNode::MakeNil());
+		
+		BranchArgs.Add(FLispNode::MakeKeyword(TEXT(":false")));
+		UEdGraphPin* FalsePin = BranchNode->GetElsePin();
+		FLispNodePtr FalseBranch = ConvertExecChainToLisp(FalsePin, Graph, VisitedNodes, bIncludePositions);
+		BranchArgs.Add(FalseBranch.IsValid() && !FalseBranch->IsNil() ? FalseBranch : FLispNode::MakeNil());
+		
+		return FLispNode::MakeList(BranchArgs);
+	}
+	
+	// Handle Variable Set
+	if (UK2Node_VariableSet* VarSet = Cast<UK2Node_VariableSet>(Node))
+	{
+		TArray<FLispNodePtr> SetArgs;
+		SetArgs.Add(FLispNode::MakeSymbol(TEXT("set")));
+		SetArgs.Add(FLispNode::MakeSymbol(VarSet->VariableReference.GetMemberName().ToString()));
+		
+		for (UEdGraphPin* Pin : VarSet->Pins)
+		{
+			if (Pin->Direction == EGPD_Input && 
+				Pin->PinType.PinCategory != UEdGraphSchema_K2::PC_Exec &&
+				Pin->PinName != UEdGraphSchema_K2::PN_Self)
+			{
+				SetArgs.Add(ConvertPureExpressionToLisp(Pin, Graph, VisitedNodes));
+				break;
+			}
+		}
+		return FLispNode::MakeList(SetArgs);
+	}
+	
+	// Handle Function Calls
+	if (UK2Node_CallFunction* CallNode = Cast<UK2Node_CallFunction>(Node))
+	{
+		UFunction* Func = CallNode->GetTargetFunction();
+		FString FuncName = Func ? Func->GetName() : GetCleanNodeName(Node);
+		
+		UEdGraphPin* TargetPin = CallNode->FindPin(UEdGraphSchema_K2::PN_Self);
+		bool bHasExplicitTarget = TargetPin && TargetPin->LinkedTo.Num() > 0;
+		
+		TArray<FLispNodePtr> CallArgs;
+		if (bHasExplicitTarget)
+		{
+			CallArgs.Add(FLispNode::MakeSymbol(TEXT("call")));
+			CallArgs.Add(ConvertPureExpressionToLisp(TargetPin, Graph, VisitedNodes));
+			CallArgs.Add(FLispNode::MakeSymbol(FuncName));
+		}
+		else
+		{
+			CallArgs.Add(FLispNode::MakeSymbol(FuncName));
+		}
+		
+		for (UEdGraphPin* Pin : CallNode->Pins)
+		{
+			if (Pin->Direction == EGPD_Input && 
+				Pin->PinType.PinCategory != UEdGraphSchema_K2::PC_Exec &&
+				Pin != TargetPin && !Pin->bHidden)
+			{
+				CallArgs.Add(ConvertPureExpressionToLisp(Pin, Graph, VisitedNodes));
+			}
+		}
+		return FLispNode::MakeList(CallArgs);
+	}
+	
+	// Handle Macro Instance - with special cases for common macros
+	if (UK2Node_MacroInstance* MacroNode = Cast<UK2Node_MacroInstance>(Node))
+	{
+		FString MacroName = MacroNode->GetMacroGraph() ? MacroNode->GetMacroGraph()->GetName() : TEXT("Unknown");
+		
+		// Special handling for ForEachLoop
+		if (MacroName.Contains(TEXT("ForEachLoop")))
+		{
+			TArray<FLispNodePtr> ForEachArgs;
+			ForEachArgs.Add(FLispNode::MakeSymbol(TEXT("foreach")));
+			
+			// Find the Array input and Array Element output
+			UEdGraphPin* ArrayPin = nullptr;
+			UEdGraphPin* ElementPin = nullptr;
+			UEdGraphPin* LoopBodyPin = nullptr;
+			
+			for (UEdGraphPin* Pin : MacroNode->Pins)
+			{
+				if (Pin->PinName.ToString().Contains(TEXT("Array")) && Pin->Direction == EGPD_Input && 
+					Pin->PinType.PinCategory != UEdGraphSchema_K2::PC_Exec)
+				{
+					ArrayPin = Pin;
+				}
+				else if (Pin->PinName.ToString().Contains(TEXT("Element")) && Pin->Direction == EGPD_Output)
+				{
+					ElementPin = Pin;
+				}
+				else if (Pin->PinName.ToString().Contains(TEXT("Loop")) && Pin->Direction == EGPD_Output &&
+					Pin->PinType.PinCategory == UEdGraphSchema_K2::PC_Exec)
+				{
+					LoopBodyPin = Pin;
+				}
+			}
+			
+			// Add element variable name
+			ForEachArgs.Add(FLispNode::MakeSymbol(TEXT("item")));
+			
+			// Add array expression
+			if (ArrayPin)
+			{
+				ForEachArgs.Add(ConvertPureExpressionToLisp(ArrayPin, Graph, VisitedNodes));
+			}
+			
+			// Add loop body
+			if (LoopBodyPin)
+			{
+				FLispNodePtr LoopBody = ConvertExecChainToLisp(LoopBodyPin, Graph, VisitedNodes, bIncludePositions);
+				if (LoopBody.IsValid() && !LoopBody->IsNil())
+				{
+					ForEachArgs.Add(LoopBody);
+				}
+			}
+			
+			return FLispNode::MakeList(ForEachArgs);
+		}
+		
+		// Special handling for Delay
+		if (MacroName.Equals(TEXT("Delay"), ESearchCase::IgnoreCase))
+		{
+			TArray<FLispNodePtr> DelayArgs;
+			DelayArgs.Add(FLispNode::MakeSymbol(TEXT("delay")));
+			
+			for (UEdGraphPin* Pin : MacroNode->Pins)
+			{
+				if (Pin->Direction == EGPD_Input && 
+					Pin->PinType.PinCategory != UEdGraphSchema_K2::PC_Exec &&
+					!Pin->bHidden)
+				{
+					DelayArgs.Add(ConvertPureExpressionToLisp(Pin, Graph, VisitedNodes));
+				}
+			}
+			
+			return FLispNode::MakeList(DelayArgs);
+		}
+		
+		// Special handling for Sequence
+		if (MacroName.Equals(TEXT("Sequence"), ESearchCase::IgnoreCase) || 
+			MacroName.Contains(TEXT("ExecutionSequence")))
+		{
+			TArray<FLispNodePtr> SeqArgs;
+			SeqArgs.Add(FLispNode::MakeSymbol(TEXT("seq")));
+			
+			// Find all "Then" output pins and convert their chains
+			TArray<UEdGraphPin*> ThenPins;
+			for (UEdGraphPin* Pin : MacroNode->Pins)
+			{
+				if (Pin->Direction == EGPD_Output && 
+					Pin->PinType.PinCategory == UEdGraphSchema_K2::PC_Exec &&
+					Pin->PinName.ToString().Contains(TEXT("Then")))
+				{
+					ThenPins.Add(Pin);
+				}
+			}
+			
+			// Sort by pin name (Then 0, Then 1, etc.)
+			ThenPins.Sort([](const UEdGraphPin& A, const UEdGraphPin& B) {
+				return A.PinName.ToString() < B.PinName.ToString();
+			});
+			
+			for (UEdGraphPin* ThenPin : ThenPins)
+			{
+				FLispNodePtr ThenBody = ConvertExecChainToLisp(ThenPin, Graph, VisitedNodes, bIncludePositions);
+				if (ThenBody.IsValid() && !ThenBody->IsNil())
+				{
+					SeqArgs.Add(ThenBody);
+				}
+			}
+			
+			return FLispNode::MakeList(SeqArgs);
+		}
+		
+		// Generic macro handling
+		TArray<FLispNodePtr> MacroArgs;
+		MacroArgs.Add(FLispNode::MakeSymbol(TEXT("macro")));
+		MacroArgs.Add(FLispNode::MakeSymbol(MacroName));
+		
+		for (UEdGraphPin* Pin : MacroNode->Pins)
+		{
+			if (Pin->Direction == EGPD_Input && 
+				Pin->PinType.PinCategory != UEdGraphSchema_K2::PC_Exec &&
+				!Pin->bHidden)
+			{
+				MacroArgs.Add(ConvertPureExpressionToLisp(Pin, Graph, VisitedNodes));
+			}
+		}
+		return FLispNode::MakeList(MacroArgs);
+	}
+	
+	// Handle Execution Sequence node (native, not macro)
+	if (UK2Node_ExecutionSequence* SeqNode = Cast<UK2Node_ExecutionSequence>(Node))
+	{
+		TArray<FLispNodePtr> SeqArgs;
+		SeqArgs.Add(FLispNode::MakeSymbol(TEXT("seq")));
+		
+		// Get all output exec pins and convert
+		for (UEdGraphPin* Pin : SeqNode->Pins)
+		{
+			if (Pin->Direction == EGPD_Output && 
+				Pin->PinType.PinCategory == UEdGraphSchema_K2::PC_Exec)
+			{
+				FLispNodePtr Body = ConvertExecChainToLisp(Pin, Graph, VisitedNodes, bIncludePositions);
+				if (Body.IsValid() && !Body->IsNil())
+				{
+					SeqArgs.Add(Body);
+				}
+			}
+		}
+		
+		return FLispNode::MakeList(SeqArgs);
+	}
+	
+	// Handle Cast
+	if (UK2Node_DynamicCast* CastNode = Cast<UK2Node_DynamicCast>(Node))
+	{
+		TArray<FLispNodePtr> CastArgs;
+		CastArgs.Add(FLispNode::MakeSymbol(TEXT("cast")));
+		CastArgs.Add(FLispNode::MakeSymbol(CastNode->TargetType ? CastNode->TargetType->GetName() : TEXT("Unknown")));
+		
+		UEdGraphPin* ObjectPin = CastNode->GetCastSourcePin();
+		if (ObjectPin)
+		{
+			CastArgs.Add(ConvertPureExpressionToLisp(ObjectPin, Graph, VisitedNodes));
+		}
+		
+		UEdGraphPin* SuccessPin = CastNode->GetValidCastPin();
+		if (SuccessPin)
+		{
+			FLispNodePtr SuccessBranch = ConvertExecChainToLisp(SuccessPin, Graph, VisitedNodes, bIncludePositions);
+			if (SuccessBranch.IsValid() && !SuccessBranch->IsNil())
+			{
+				CastArgs.Add(SuccessBranch);
+			}
+		}
+		return FLispNode::MakeList(CastArgs);
+	}
+	
+	// Handle Switch Integer
+	if (UK2Node_SwitchInteger* SwitchIntNode = Cast<UK2Node_SwitchInteger>(Node))
+	{
+		TArray<FLispNodePtr> SwitchArgs;
+		SwitchArgs.Add(FLispNode::MakeSymbol(TEXT("switch-int")));
+		
+		// Add the selection value
+		UEdGraphPin* SelectionPin = SwitchIntNode->GetSelectionPin();
+		if (SelectionPin)
+		{
+			SwitchArgs.Add(ConvertPureExpressionToLisp(SelectionPin, Graph, VisitedNodes));
+		}
+		
+		// Add each case
+		for (UEdGraphPin* Pin : SwitchIntNode->Pins)
+		{
+			if (Pin->Direction == EGPD_Output && 
+				Pin->PinType.PinCategory == UEdGraphSchema_K2::PC_Exec &&
+				Pin != SwitchIntNode->GetDefaultPin())
+			{
+				// Get the case value from pin name
+				FString PinName = Pin->PinName.ToString();
+				SwitchArgs.Add(FLispNode::MakeKeyword(TEXT(":") + PinName));
+				
+				FLispNodePtr CaseBody = ConvertExecChainToLisp(Pin, Graph, VisitedNodes, bIncludePositions);
+				SwitchArgs.Add(CaseBody.IsValid() && !CaseBody->IsNil() ? CaseBody : FLispNode::MakeNil());
+			}
+		}
+		
+		// Add default case
+		UEdGraphPin* DefaultPin = SwitchIntNode->GetDefaultPin();
+		if (DefaultPin && DefaultPin->LinkedTo.Num() > 0)
+		{
+			SwitchArgs.Add(FLispNode::MakeKeyword(TEXT(":default")));
+			FLispNodePtr DefaultBody = ConvertExecChainToLisp(DefaultPin, Graph, VisitedNodes, bIncludePositions);
+			SwitchArgs.Add(DefaultBody.IsValid() && !DefaultBody->IsNil() ? DefaultBody : FLispNode::MakeNil());
+		}
+		
+		return FLispNode::MakeList(SwitchArgs);
+	}
+	
+	// Handle Switch String
+	if (UK2Node_SwitchString* SwitchStrNode = Cast<UK2Node_SwitchString>(Node))
+	{
+		TArray<FLispNodePtr> SwitchArgs;
+		SwitchArgs.Add(FLispNode::MakeSymbol(TEXT("switch-string")));
+		
+		UEdGraphPin* SelectionPin = SwitchStrNode->GetSelectionPin();
+		if (SelectionPin)
+		{
+			SwitchArgs.Add(ConvertPureExpressionToLisp(SelectionPin, Graph, VisitedNodes));
+		}
+		
+		for (UEdGraphPin* Pin : SwitchStrNode->Pins)
+		{
+			if (Pin->Direction == EGPD_Output && 
+				Pin->PinType.PinCategory == UEdGraphSchema_K2::PC_Exec &&
+				Pin != SwitchStrNode->GetDefaultPin())
+			{
+				FString PinName = Pin->PinName.ToString();
+				SwitchArgs.Add(FLispNode::MakeKeyword(TEXT(":\"") + PinName + TEXT("\"")));
+				
+				FLispNodePtr CaseBody = ConvertExecChainToLisp(Pin, Graph, VisitedNodes, bIncludePositions);
+				SwitchArgs.Add(CaseBody.IsValid() && !CaseBody->IsNil() ? CaseBody : FLispNode::MakeNil());
+			}
+		}
+		
+		UEdGraphPin* DefaultPin = SwitchStrNode->GetDefaultPin();
+		if (DefaultPin && DefaultPin->LinkedTo.Num() > 0)
+		{
+			SwitchArgs.Add(FLispNode::MakeKeyword(TEXT(":default")));
+			FLispNodePtr DefaultBody = ConvertExecChainToLisp(DefaultPin, Graph, VisitedNodes, bIncludePositions);
+			SwitchArgs.Add(DefaultBody.IsValid() && !DefaultBody->IsNil() ? DefaultBody : FLispNode::MakeNil());
+		}
+		
+		return FLispNode::MakeList(SwitchArgs);
+	}
+	
+	// Handle Switch Enum
+	if (UK2Node_SwitchEnum* SwitchEnumNode = Cast<UK2Node_SwitchEnum>(Node))
+	{
+		TArray<FLispNodePtr> SwitchArgs;
+		SwitchArgs.Add(FLispNode::MakeSymbol(TEXT("switch-enum")));
+		
+		// Add the enum type
+		if (SwitchEnumNode->Enum)
+		{
+			SwitchArgs.Add(FLispNode::MakeSymbol(SwitchEnumNode->Enum->GetName()));
+		}
+		
+		UEdGraphPin* SelectionPin = SwitchEnumNode->GetSelectionPin();
+		if (SelectionPin)
+		{
+			SwitchArgs.Add(ConvertPureExpressionToLisp(SelectionPin, Graph, VisitedNodes));
+		}
+		
+		for (UEdGraphPin* Pin : SwitchEnumNode->Pins)
+		{
+			if (Pin->Direction == EGPD_Output && 
+				Pin->PinType.PinCategory == UEdGraphSchema_K2::PC_Exec &&
+				Pin != SwitchEnumNode->GetDefaultPin())
+			{
+				FString PinName = Pin->PinName.ToString();
+				SwitchArgs.Add(FLispNode::MakeKeyword(TEXT(":") + PinName));
+				
+				FLispNodePtr CaseBody = ConvertExecChainToLisp(Pin, Graph, VisitedNodes, bIncludePositions);
+				SwitchArgs.Add(CaseBody.IsValid() && !CaseBody->IsNil() ? CaseBody : FLispNode::MakeNil());
+			}
+		}
+		
+		UEdGraphPin* DefaultPin = SwitchEnumNode->GetDefaultPin();
+		if (DefaultPin && DefaultPin->LinkedTo.Num() > 0)
+		{
+			SwitchArgs.Add(FLispNode::MakeKeyword(TEXT(":default")));
+			FLispNodePtr DefaultBody = ConvertExecChainToLisp(DefaultPin, Graph, VisitedNodes, bIncludePositions);
+			SwitchArgs.Add(DefaultBody.IsValid() && !DefaultBody->IsNil() ? DefaultBody : FLispNode::MakeNil());
+		}
+		
+		return FLispNode::MakeList(SwitchArgs);
+	}
+	
+	// Handle Comment nodes - skip them in execution flow but include in output
+	if (UEdGraphNode_Comment* CommentNode = Cast<UEdGraphNode_Comment>(Node))
+	{
+		// Comments are not part of execution flow, but we include them as annotations
+		TArray<FLispNodePtr> CommentArgs;
+		CommentArgs.Add(FLispNode::MakeSymbol(TEXT("comment")));
+		CommentArgs.Add(FLispNode::MakeString(CommentNode->NodeComment));
+		return FLispNode::MakeList(CommentArgs);
+	}
+	
+	// Handle Input Action
+	if (UK2Node_InputAction* InputActionNode = Cast<UK2Node_InputAction>(Node))
+	{
+		TArray<FLispNodePtr> InputArgs;
+		InputArgs.Add(FLispNode::MakeSymbol(TEXT("on-input")));
+		InputArgs.Add(FLispNode::MakeSymbol(InputActionNode->InputActionName.ToString()));
+		
+		// Check which exec pins have connections
+		for (UEdGraphPin* Pin : InputActionNode->Pins)
+		{
+			if (Pin->Direction == EGPD_Output && 
+				Pin->PinType.PinCategory == UEdGraphSchema_K2::PC_Exec &&
+				Pin->LinkedTo.Num() > 0)
+			{
+				FString EventType = Pin->PinName.ToString();
+				InputArgs.Add(FLispNode::MakeKeyword(TEXT(":") + EventType.ToLower()));
+				
+				FLispNodePtr Body = ConvertExecChainToLisp(Pin, Graph, VisitedNodes, bIncludePositions);
+				InputArgs.Add(Body.IsValid() && !Body->IsNil() ? Body : FLispNode::MakeNil());
+			}
+		}
+		
+		return FLispNode::MakeList(InputArgs);
+	}
+	
+	// Generic fallback
+	TArray<FLispNodePtr> GenericArgs;
+	GenericArgs.Add(FLispNode::MakeSymbol(GetCleanNodeName(Node)));
+	
+	for (UEdGraphPin* Pin : Node->Pins)
+	{
+		if (Pin->Direction == EGPD_Input && 
+			Pin->PinType.PinCategory != UEdGraphSchema_K2::PC_Exec &&
+			!Pin->bHidden)
+		{
+			GenericArgs.Add(ConvertPureExpressionToLisp(Pin, Graph, VisitedNodes));
+		}
+	}
+	return FLispNode::MakeList(GenericArgs);
+}
+
+static FLispNodePtr ConvertExecChainToLisp(UEdGraphPin* ExecPin, UEdGraph* Graph, TSet<UEdGraphNode*>& VisitedNodes, bool bIncludePositions)
+{
+	if (!ExecPin || ExecPin->LinkedTo.Num() == 0)
+	{
+		return FLispNode::MakeNil();
+	}
+	
+	TArray<FLispNodePtr> Statements;
+	UEdGraphPin* CurrentExecPin = ExecPin;
+	
+	while (CurrentExecPin && CurrentExecPin->LinkedTo.Num() > 0)
+	{
+		UEdGraphNode* NextNode = CurrentExecPin->LinkedTo[0]->GetOwningNode();
+		
+		if (!NextNode || VisitedNodes.Contains(NextNode))
+		{
+			break;
+		}
+		
+		VisitedNodes.Add(NextNode);
+		
+		FLispNodePtr NodeLisp = ConvertNodeToLisp(NextNode, Graph, VisitedNodes, bIncludePositions);
+		if (NodeLisp.IsValid() && !NodeLisp->IsNil())
+		{
+			Statements.Add(NodeLisp);
+		}
+		
+		CurrentExecPin = GetThenPin(NextNode);
+		
+		// For branch nodes, we've already handled the branches in ConvertNodeToLisp
+		if (Cast<UK2Node_IfThenElse>(NextNode))
+		{
+			break;
+		}
+	}
+	
+	if (Statements.Num() == 0)
+	{
+		return FLispNode::MakeNil();
+	}
+	if (Statements.Num() == 1)
+	{
+		return Statements[0];
+	}
+	
+	TArray<FLispNodePtr> SeqArgs;
+	SeqArgs.Add(FLispNode::MakeSymbol(TEXT("seq")));
+	SeqArgs.Append(Statements);
+	return FLispNode::MakeList(SeqArgs);
+}
+
+static FLispNodePtr ConvertEventToLisp(UK2Node_Event* EventNode, UEdGraph* Graph, bool bIncludePositions)
+{
+	TSet<UEdGraphNode*> VisitedNodes;
+	VisitedNodes.Add(EventNode);
+	
+	TArray<FLispNodePtr> EventArgs;
+	EventArgs.Add(FLispNode::MakeSymbol(TEXT("event")));
+	
+	FString EventName = EventNode->GetFunctionName().ToString();
+	if (EventName.StartsWith(TEXT("ReceiveBeginPlay"))) EventName = TEXT("BeginPlay");
+	else if (EventName.StartsWith(TEXT("ReceiveTick"))) EventName = TEXT("Tick");
+	else if (EventName.StartsWith(TEXT("ReceiveEndPlay"))) EventName = TEXT("EndPlay");
+	EventArgs.Add(FLispNode::MakeSymbol(EventName));
+	
+	// Parameters
+	TArray<FLispNodePtr> ParamsList;
+	for (UEdGraphPin* Pin : EventNode->Pins)
+	{
+		if (Pin->Direction == EGPD_Output && 
+			Pin->PinType.PinCategory != UEdGraphSchema_K2::PC_Exec &&
+			Pin->PinType.PinCategory != UEdGraphSchema_K2::PC_Delegate &&
+			!Pin->bHidden)
+		{
+			TArray<FLispNodePtr> ParamDef;
+			ParamDef.Add(FLispNode::MakeSymbol(Pin->PinName.ToString()));
+			ParamDef.Add(FLispNode::MakeSymbol(PinTypeToLispType(Pin->PinType)));
+			ParamsList.Add(FLispNode::MakeList(ParamDef));
+		}
+	}
+	
+	if (ParamsList.Num() > 0)
+	{
+		EventArgs.Add(FLispNode::MakeKeyword(TEXT(":params")));
+		EventArgs.Add(FLispNode::MakeList(ParamsList));
+	}
+	
+	UEdGraphPin* ThenPin = GetThenPin(EventNode);
+	FLispNodePtr Body = ConvertExecChainToLisp(ThenPin, Graph, VisitedNodes, bIncludePositions);
+	
+	if (Body.IsValid() && !Body->IsNil())
+	{
+		if (Body->IsList() && Body->GetFormName() == TEXT("seq"))
+		{
+			for (int32 i = 1; i < Body->Num(); i++)
+			{
+				EventArgs.Add(Body->Get(i));
+			}
+		}
+		else
+		{
+			EventArgs.Add(Body);
+		}
+	}
+	
+	return FLispNode::MakeList(EventArgs);
+}
+
+static FLispNodePtr ConvertCustomEventToLisp(UK2Node_CustomEvent* EventNode, UEdGraph* Graph, bool bIncludePositions)
+{
+	TSet<UEdGraphNode*> VisitedNodes;
+	VisitedNodes.Add(EventNode);
+	
+	TArray<FLispNodePtr> EventArgs;
+	EventArgs.Add(FLispNode::MakeSymbol(TEXT("event")));
+	EventArgs.Add(FLispNode::MakeSymbol(EventNode->CustomFunctionName.ToString()));
+	
+	TArray<FLispNodePtr> ParamsList;
+	for (UEdGraphPin* Pin : EventNode->Pins)
+	{
+		if (Pin->Direction == EGPD_Output && 
+			Pin->PinType.PinCategory != UEdGraphSchema_K2::PC_Exec &&
+			Pin->PinType.PinCategory != UEdGraphSchema_K2::PC_Delegate &&
+			!Pin->bHidden)
+		{
+			TArray<FLispNodePtr> ParamDef;
+			ParamDef.Add(FLispNode::MakeSymbol(Pin->PinName.ToString()));
+			ParamDef.Add(FLispNode::MakeSymbol(PinTypeToLispType(Pin->PinType)));
+			ParamsList.Add(FLispNode::MakeList(ParamDef));
+		}
+	}
+	
+	if (ParamsList.Num() > 0)
+	{
+		EventArgs.Add(FLispNode::MakeKeyword(TEXT(":params")));
+		EventArgs.Add(FLispNode::MakeList(ParamsList));
+	}
+	
+	UEdGraphPin* ThenPin = GetThenPin(EventNode);
+	FLispNodePtr Body = ConvertExecChainToLisp(ThenPin, Graph, VisitedNodes, bIncludePositions);
+	
+	if (Body.IsValid() && !Body->IsNil())
+	{
+		if (Body->IsList() && Body->GetFormName() == TEXT("seq"))
+		{
+			for (int32 i = 1; i < Body->Num(); i++)
+			{
+				EventArgs.Add(Body->Get(i));
+			}
+		}
+		else
+		{
+			EventArgs.Add(Body);
+		}
+	}
+	
+	return FLispNode::MakeList(EventArgs);
+}
+
+// Convert component-bound event to Lisp
+// Example: OnComponentBeginOverlap on BoxCollision component
+// Output: (on-component BoxCollision BeginOverlap :params ((OtherActor Actor) ...) body...)
+static FLispNodePtr ConvertComponentEventToLisp(UK2Node_ComponentBoundEvent* CompEvent, UEdGraph* Graph, bool bIncludePositions)
+{
+	TSet<UEdGraphNode*> VisitedNodes;
+	VisitedNodes.Add(CompEvent);
+	
+	TArray<FLispNodePtr> EventArgs;
+	EventArgs.Add(FLispNode::MakeSymbol(TEXT("on-component")));
+	
+	// Component name
+	FString ComponentName = CompEvent->ComponentPropertyName.ToString();
+	EventArgs.Add(FLispNode::MakeSymbol(ComponentName));
+	
+	// Delegate/event name (e.g., "OnComponentBeginOverlap" -> "BeginOverlap")
+	FString DelegateName = CompEvent->DelegatePropertyName.ToString();
+	DelegateName.RemoveFromStart(TEXT("On"));
+	DelegateName.RemoveFromStart(TEXT("Component"));
+	EventArgs.Add(FLispNode::MakeSymbol(DelegateName));
+	
+	// Collect parameters (output pins that aren't exec)
+	TArray<FLispNodePtr> ParamsList;
+	for (UEdGraphPin* Pin : CompEvent->Pins)
+	{
+		if (Pin->Direction == EGPD_Output && 
+			Pin->PinType.PinCategory != UEdGraphSchema_K2::PC_Exec &&
+			Pin->PinType.PinCategory != UEdGraphSchema_K2::PC_Delegate &&
+			!Pin->bHidden)
+		{
+			TArray<FLispNodePtr> ParamDef;
+			ParamDef.Add(FLispNode::MakeSymbol(Pin->PinName.ToString()));
+			ParamDef.Add(FLispNode::MakeSymbol(PinTypeToLispType(Pin->PinType)));
+			ParamsList.Add(FLispNode::MakeList(ParamDef));
+		}
+	}
+	
+	if (ParamsList.Num() > 0)
+	{
+		EventArgs.Add(FLispNode::MakeKeyword(TEXT(":params")));
+		EventArgs.Add(FLispNode::MakeList(ParamsList));
+	}
+	
+	// Convert the body
+	UEdGraphPin* ThenPin = GetThenPin(CompEvent);
+	FLispNodePtr Body = ConvertExecChainToLisp(ThenPin, Graph, VisitedNodes, bIncludePositions);
+	
+	if (Body.IsValid() && !Body->IsNil())
+	{
+		if (Body->IsList() && Body->GetFormName() == TEXT("seq"))
+		{
+			for (int32 i = 1; i < Body->Num(); i++)
+			{
+				EventArgs.Add(Body->Get(i));
+			}
+		}
+		else
+		{
+			EventArgs.Add(Body);
+		}
+	}
+	
+	return FLispNode::MakeList(EventArgs);
+}
+
+//------------------------------------------------------------------------------
+// FECACommand_ParseBlueprintLisp
+//------------------------------------------------------------------------------
+
+FECACommandResult FECACommand_ParseBlueprintLisp::Execute(const TSharedPtr<FJsonObject>& Params)
+{
+	FString Code;
+	if (!GetStringParam(Params, TEXT("code"), Code))
+	{
+		return FECACommandResult::Error(TEXT("Missing required parameter: code"));
+	}
+	
+	bool bPrettyPrint = false;
+	Params->TryGetBoolField(TEXT("pretty_print"), bPrettyPrint);
+	
+	FLispParseResult ParseResult = FLispParser::Parse(Code);
+	
+	TSharedPtr<FJsonObject> Result = MakeResult();
+	
+	if (!ParseResult.bSuccess)
+	{
+		Result->SetBoolField(TEXT("success"), false);
+		Result->SetStringField(TEXT("error"), ParseResult.Error);
+		Result->SetNumberField(TEXT("error_line"), ParseResult.ErrorLine);
+		Result->SetNumberField(TEXT("error_column"), ParseResult.ErrorColumn);
+		return FECACommandResult::Success(Result);
+	}
+	
+	Result->SetBoolField(TEXT("success"), true);
+	Result->SetNumberField(TEXT("expression_count"), ParseResult.Nodes.Num());
+	
+	TArray<TSharedPtr<FJsonValue>> AstArray;
+	for (const auto& Node : ParseResult.Nodes)
+	{
+		AstArray.Add(MakeShared<FJsonValueObject>(LispNodeToJson(Node)));
+	}
+	Result->SetArrayField(TEXT("ast"), AstArray);
+	
+	if (bPrettyPrint)
+	{
+		FString Pretty;
+		for (int32 i = 0; i < ParseResult.Nodes.Num(); i++)
+		{
+			if (i > 0) Pretty += TEXT("\n\n");
+			Pretty += ParseResult.Nodes[i]->ToString(true, 0);
+		}
+		Result->SetStringField(TEXT("pretty"), Pretty);
+	}
+	
+	TSet<FString> Symbols;
+	TFunction<void(const FLispNodePtr&)> CollectSymbols = [&](const FLispNodePtr& Node)
+	{
+		if (!Node.IsValid()) return;
+		if (Node->IsSymbol()) Symbols.Add(Node->StringValue);
+		else if (Node->IsList())
+		{
+			for (const auto& Child : Node->Children) CollectSymbols(Child);
+		}
+	};
+	
+	for (const auto& Node : ParseResult.Nodes) CollectSymbols(Node);
+	
+	TArray<TSharedPtr<FJsonValue>> SymbolsArray;
+	for (const FString& Sym : Symbols)
+	{
+		SymbolsArray.Add(MakeShared<FJsonValueString>(Sym));
+	}
+	Result->SetArrayField(TEXT("symbols"), SymbolsArray);
+	
+	return FECACommandResult::Success(Result);
+}
+
+//------------------------------------------------------------------------------
+// FECACommand_BlueprintToLisp
+//------------------------------------------------------------------------------
+
+FECACommandResult FECACommand_BlueprintToLisp::Execute(const TSharedPtr<FJsonObject>& Params)
+{
+	FString BlueprintPath;
+	if (!GetStringParam(Params, TEXT("blueprint_path"), BlueprintPath))
+	{
+		return FECACommandResult::Error(TEXT("Missing required parameter: blueprint_path"));
+	}
+	
+	FString GraphName = TEXT("EventGraph");
+	GetStringParam(Params, TEXT("graph_name"), GraphName, false);
+	
+	bool bIncludeComments = true;
+	Params->TryGetBoolField(TEXT("include_comments"), bIncludeComments);
+	
+	bool bIncludePositions = false;
+	Params->TryGetBoolField(TEXT("include_positions"), bIncludePositions);
+	
+	UBlueprint* Blueprint = LoadBlueprintByPathLisp(BlueprintPath);
+	if (!Blueprint)
+	{
+		return FECACommandResult::Error(FString::Printf(TEXT("Blueprint not found: %s"), *BlueprintPath));
+	}
+	
+	UEdGraph* Graph = FindGraphByNameLisp(Blueprint, GraphName);
+	if (!Graph)
+	{
+		return FECACommandResult::Error(FString::Printf(TEXT("Graph not found: %s"), *GraphName));
+	}
+	
+	TArray<UK2Node_Event*> EventNodes;
+	TArray<UK2Node_CustomEvent*> CustomEventNodes;
+	TArray<UK2Node_ComponentBoundEvent*> ComponentEventNodes;
+	
+	for (UEdGraphNode* Node : Graph->Nodes)
+	{
+		// Check for component-bound events first (they inherit from UK2Node_Event)
+		if (UK2Node_ComponentBoundEvent* CompEvent = Cast<UK2Node_ComponentBoundEvent>(Node))
+		{
+			ComponentEventNodes.Add(CompEvent);
+		}
+		else if (UK2Node_Event* EventNode = Cast<UK2Node_Event>(Node))
+		{
+			EventNodes.Add(EventNode);
+		}
+		else if (UK2Node_CustomEvent* CustomEvent = Cast<UK2Node_CustomEvent>(Node))
+		{
+			CustomEventNodes.Add(CustomEvent);
+		}
+	}
+	
+	TArray<FLispNodePtr> TopLevelForms;
+	
+	for (UK2Node_Event* EventNode : EventNodes)
+	{
+		FLispNodePtr EventLisp = ConvertEventToLisp(EventNode, Graph, bIncludePositions);
+		if (EventLisp.IsValid() && !EventLisp->IsNil())
+		{
+			TopLevelForms.Add(EventLisp);
+		}
+	}
+	
+	for (UK2Node_CustomEvent* CustomEvent : CustomEventNodes)
+	{
+		FLispNodePtr EventLisp = ConvertCustomEventToLisp(CustomEvent, Graph, bIncludePositions);
+		if (EventLisp.IsValid() && !EventLisp->IsNil())
+		{
+			TopLevelForms.Add(EventLisp);
+		}
+	}
+	
+	// Convert component-bound events
+	for (UK2Node_ComponentBoundEvent* CompEvent : ComponentEventNodes)
+	{
+		FLispNodePtr EventLisp = ConvertComponentEventToLisp(CompEvent, Graph, bIncludePositions);
+		if (EventLisp.IsValid() && !EventLisp->IsNil())
+		{
+			TopLevelForms.Add(EventLisp);
+		}
+	}
+	
+	FString LispCode;
+	if (bIncludeComments)
+	{
+		LispCode = FString::Printf(TEXT("; Blueprint: %s\n; Graph: %s\n\n"), *Blueprint->GetName(), *GraphName);
+	}
+	
+	for (int32 i = 0; i < TopLevelForms.Num(); i++)
+	{
+		if (i > 0) LispCode += TEXT("\n\n");
+		LispCode += TopLevelForms[i]->ToString(true, 0);
+	}
+	
+	TSharedPtr<FJsonObject> Result = MakeResult();
+	Result->SetBoolField(TEXT("success"), true);
+	Result->SetStringField(TEXT("blueprint_path"), BlueprintPath);
+	Result->SetStringField(TEXT("graph_name"), GraphName);
+	Result->SetStringField(TEXT("lisp_code"), LispCode);
+	Result->SetNumberField(TEXT("event_count"), EventNodes.Num() + CustomEventNodes.Num());
+	
+	return FECACommandResult::Success(Result);
+}
+
+//------------------------------------------------------------------------------
+// Lisp -> Blueprint Conversion Helpers
+//------------------------------------------------------------------------------
+
+// Context for Lisp to Blueprint conversion
+struct FLispToBPContext
+{
+	UBlueprint* Blueprint = nullptr;
+	UEdGraph* Graph = nullptr;
+	TMap<FString, UEdGraphNode*> TempIdToNode;
+	TMap<FString, FString> VariableToNodeId; // Variable name -> node ID that outputs it
+	TMap<FString, FString> VariableToPin;    // Variable name -> output pin name
+	TArray<FString> Errors;
+	TArray<FString> Warnings;
+	int32 NextTempId = 0;
+	int32 CurrentX = 0;
+	int32 CurrentY = 0;
+	FString LastAssetPath; // Used to pass asset path from ResolveLispExpression to caller
+	
+	// Helper to clear the last asset path after use
+	void ClearLastAssetPath() { LastAssetPath.Empty(); }
+	
+	FString GenerateTempId()
+	{
+		return FString::Printf(TEXT("_temp_%d"), NextTempId++);
+	}
+	
+	void AdvancePosition()
+	{
+		CurrentX += 400;
+	}
+	
+	void NewRow()
+	{
+		CurrentX = 0;
+		CurrentY += 200;
+	}
+};
+
+// Forward declarations
+static UEdGraphNode* ConvertLispFormToNode(const FLispNodePtr& Form, FLispToBPContext& Ctx, UEdGraphPin*& OutExecPin);
+static bool ConnectNodes(UEdGraphPin* SourcePin, UEdGraphPin* TargetPin, FLispToBPContext& Ctx);
+static UEdGraphPin* FindOutputPin(UEdGraphNode* Node, const FString& PinName);
+static UEdGraphPin* FindInputPin(UEdGraphNode* Node, const FString& PinName);
+static UEdGraphPin* GetNodeExecOutput(UEdGraphNode* Node);
+static UEdGraphPin* GetNodeExecInput(UEdGraphNode* Node);
+
+// Helper to ensure a node has a valid GUID
+// Some nodes may not have their GUID properly initialized when created programmatically
+static void EnsureNodeHasValidGuid(UEdGraphNode* Node)
+{
+	if (Node && !Node->NodeGuid.IsValid())
+	{
+		Node->CreateNewGuid();
+	}
+}
+
+// Helper macro to allocate pins and ensure GUID in one step
+#define ALLOCATE_PINS_AND_GUID(Node) do { (Node)->AllocateDefaultPins(); EnsureNodeHasValidGuid(Node); } while(0)
+
+// Find an output pin by name (or first non-exec output)
+static UEdGraphPin* FindOutputPin(UEdGraphNode* Node, const FString& PinName)
+{
+	if (!Node) return nullptr;
+	
+	// Try exact match first
+	for (UEdGraphPin* Pin : Node->Pins)
+	{
+		if (Pin->Direction == EGPD_Output && Pin->PinName.ToString() == PinName)
+		{
+			return Pin;
+		}
+	}
+	
+	// Try case-insensitive
+	for (UEdGraphPin* Pin : Node->Pins)
+	{
+		if (Pin->Direction == EGPD_Output && Pin->PinName.ToString().Equals(PinName, ESearchCase::IgnoreCase))
+		{
+			return Pin;
+		}
+	}
+	
+	// Return first non-exec output
+	for (UEdGraphPin* Pin : Node->Pins)
+	{
+		if (Pin->Direction == EGPD_Output && 
+			Pin->PinType.PinCategory != UEdGraphSchema_K2::PC_Exec &&
+			!Pin->bHidden)
+		{
+			return Pin;
+		}
+	}
+	
+	return nullptr;
+}
+
+// Find an input pin by name
+static UEdGraphPin* FindInputPin(UEdGraphNode* Node, const FString& PinName)
+{
+	if (!Node) return nullptr;
+	
+	for (UEdGraphPin* Pin : Node->Pins)
+	{
+		if (Pin->Direction == EGPD_Input && Pin->PinName.ToString().Equals(PinName, ESearchCase::IgnoreCase))
+		{
+			return Pin;
+		}
+	}
+	
+	return nullptr;
+}
+
+// Get the execution output pin (then pin)
+static UEdGraphPin* GetNodeExecOutput(UEdGraphNode* Node)
+{
+	if (!Node) return nullptr;
+	
+	for (UEdGraphPin* Pin : Node->Pins)
+	{
+		if (Pin->Direction == EGPD_Output && Pin->PinType.PinCategory == UEdGraphSchema_K2::PC_Exec)
+		{
+			return Pin;
+		}
+	}
+	return nullptr;
+}
+
+// Get the execution input pin
+static UEdGraphPin* GetNodeExecInput(UEdGraphNode* Node)
+{
+	if (!Node) return nullptr;
+	
+	for (UEdGraphPin* Pin : Node->Pins)
+	{
+		if (Pin->Direction == EGPD_Input && Pin->PinType.PinCategory == UEdGraphSchema_K2::PC_Exec)
+		{
+			return Pin;
+		}
+	}
+	return nullptr;
+}
+
+// Forward declaration
+static UEdGraphPin* ResolveLispExpression(const FLispNodePtr& Expr, FLispToBPContext& Ctx);
+
+// Helper to set a pin value from a Lisp expression (handles assets, literals, etc.)
+static bool SetPinValueFromExpr(UEdGraphPin* Pin, const FLispNodePtr& Expr, FLispToBPContext& Ctx)
+{
+	if (!Pin || !Expr.IsValid()) return false;
+	
+	// Clear any previous asset path
+	Ctx.ClearLastAssetPath();
+	
+	// Try to resolve as a connected expression first
+	UEdGraphPin* SourcePin = ResolveLispExpression(Expr, Ctx);
+	if (SourcePin)
+	{
+		// It's a connected node output
+		SourcePin->MakeLinkTo(Pin);
+		return true;
+	}
+	
+	// Check if ResolveLispExpression set an asset path
+	if (!Ctx.LastAssetPath.IsEmpty())
+	{
+		// Load the asset and set it as DefaultObject
+		UObject* Asset = StaticLoadObject(UObject::StaticClass(), nullptr, *Ctx.LastAssetPath);
+		if (Asset)
+		{
+			Pin->DefaultObject = Asset;
+			Ctx.ClearLastAssetPath();
+			return true;
+		}
+		else
+		{
+			Ctx.Warnings.Add(FString::Printf(TEXT("Could not load asset: %s"), *Ctx.LastAssetPath));
+			Ctx.ClearLastAssetPath();
+		}
+	}
+	
+	// Handle literal values
+	if (Expr->IsNumber())
+	{
+		SetNumericPinDefaultValue(Pin, Expr->NumberValue);
+		return true;
+	}
+	
+	if (Expr->IsString())
+	{
+		Pin->DefaultValue = Expr->StringValue;
+		return true;
+	}
+	
+	if (Expr->IsSymbol())
+	{
+		FString SymValue = Expr->StringValue;
+		if (SymValue.Equals(TEXT("true"), ESearchCase::IgnoreCase))
+		{
+			Pin->DefaultValue = TEXT("true");
+			return true;
+		}
+		if (SymValue.Equals(TEXT("false"), ESearchCase::IgnoreCase))
+		{
+			Pin->DefaultValue = TEXT("false");
+			return true;
+		}
+		if (SymValue.Equals(TEXT("nil"), ESearchCase::IgnoreCase))
+		{
+			Pin->DefaultValue = TEXT("");
+			return true;
+		}
+		// Try as enum or other string value
+		Pin->DefaultValue = SymValue;
+		return true;
+	}
+	
+	return false;
+}
+
+// Connect two pins
+static bool ConnectNodes(UEdGraphPin* SourcePin, UEdGraphPin* TargetPin, FLispToBPContext& Ctx)
+{
+	if (!SourcePin || !TargetPin)
+	{
+		Ctx.Errors.Add(TEXT("Cannot connect: null pin"));
+		return false;
+	}
+	
+	Ctx.Warnings.Add(FString::Printf(TEXT("DEBUG ConnectNodes: Connecting '%s' (%s, dir=%d) -> '%s' (%s, dir=%d)"),
+		*SourcePin->PinName.ToString(),
+		*SourcePin->GetOwningNode()->GetClass()->GetName(),
+		(int32)SourcePin->Direction,
+		*TargetPin->PinName.ToString(),
+		*TargetPin->GetOwningNode()->GetClass()->GetName(),
+		(int32)TargetPin->Direction));
+	
+	// Verify pin directions are compatible
+	// Source should be Output, Target should be Input (or vice versa for bidirectional)
+	if (SourcePin->Direction == TargetPin->Direction)
+	{
+		Ctx.Warnings.Add(FString::Printf(TEXT("DEBUG ConnectNodes: WARNING - Same direction pins! Source=%d, Target=%d"), 
+			(int32)SourcePin->Direction, (int32)TargetPin->Direction));
+	}
+	
+	// For exec pins (output), we can only have one connection in most cases
+	// So break existing connections first if this is an exec output
+	if (SourcePin->PinType.PinCategory == UEdGraphSchema_K2::PC_Exec && 
+		SourcePin->Direction == EGPD_Output &&
+		SourcePin->LinkedTo.Num() > 0)
+	{
+		Ctx.Warnings.Add(TEXT("DEBUG ConnectNodes: Breaking existing exec connections"));
+		SourcePin->BreakAllPinLinks();
+	}
+	
+	// Check if already connected to this specific target
+	if (SourcePin->LinkedTo.Contains(TargetPin))
+	{
+		Ctx.Warnings.Add(TEXT("DEBUG ConnectNodes: Already connected"));
+		return true;
+	}
+	
+	// Make the connection
+	SourcePin->MakeLinkTo(TargetPin);
+	
+	// Verify the connection was made
+	bool bSourceHasTarget = SourcePin->LinkedTo.Contains(TargetPin);
+	bool bTargetHasSource = TargetPin->LinkedTo.Contains(SourcePin);
+	
+	Ctx.Warnings.Add(FString::Printf(TEXT("DEBUG ConnectNodes: MakeLinkTo returned, SourceHasTarget=%s, TargetHasSource=%s, SourceLinkedTo=%d, TargetLinkedTo=%d"),
+		bSourceHasTarget ? TEXT("YES") : TEXT("NO"),
+		bTargetHasSource ? TEXT("YES") : TEXT("NO"),
+		SourcePin->LinkedTo.Num(),
+		TargetPin->LinkedTo.Num()));
+	
+	return bSourceHasTarget && bTargetHasSource;
+}
+
+// Resolve a Lisp expression to a node/pin that produces the value
+static UEdGraphPin* ResolveLispExpression(const FLispNodePtr& Expr, FLispToBPContext& Ctx)
+{
+	if (!Expr.IsValid() || Expr->IsNil())
+	{
+		return nullptr;
+	}
+	
+	// Symbol - could be a variable reference
+	if (Expr->IsSymbol())
+	{
+		FString SymName = Expr->StringValue;
+		
+		// Check if it's a known variable (from let bindings or event parameters)
+		if (Ctx.VariableToNodeId.Contains(SymName))
+		{
+			FString NodeId = Ctx.VariableToNodeId[SymName];
+			FString PinName = Ctx.VariableToPin.Contains(SymName) ? Ctx.VariableToPin[SymName] : TEXT("");
+			
+			// Check if this is a literal value (stored with _literal_ prefix)
+			// PinName contains the literal value string (e.g., "0.208")
+			if (NodeId.StartsWith(TEXT("_literal_")))
+			{
+				FString LiteralValueStr = PinName; // PinName holds the literal value
+				
+				// Check if we already created a node for this literal (to avoid duplicates)
+				FString LiteralNodeKey = TEXT("_literalnode_") + SymName;
+				if (UEdGraphNode** ExistingNode = Ctx.TempIdToNode.Find(LiteralNodeKey))
+				{
+					return FindOutputPin(*ExistingNode, TEXT("ReturnValue"));
+				}
+				
+				// Create a pure "float * 1" or use Conv_DoubleToDouble as identity function
+				// to create a node that outputs the literal value
+				// Actually, the cleanest way is to use a Multiply node with the literal and 1.0
+				UFunction* MultiplyFunc = UKismetMathLibrary::StaticClass()->FindFunctionByName(TEXT("Multiply_DoubleDouble"));
+				if (!MultiplyFunc)
+				{
+					MultiplyFunc = UKismetMathLibrary::StaticClass()->FindFunctionByName(TEXT("Multiply_FloatFloat"));
+				}
+				
+				if (MultiplyFunc)
+				{
+					UK2Node_CallFunction* LiteralNode = NewObject<UK2Node_CallFunction>(Ctx.Graph);
+					LiteralNode->SetFromFunction(MultiplyFunc);
+					LiteralNode->NodePosX = Ctx.CurrentX;
+					LiteralNode->NodePosY = Ctx.CurrentY;
+					Ctx.Graph->AddNode(LiteralNode, false, false);
+					LiteralNode->AllocateDefaultPins();
+					EnsureNodeHasValidGuid(LiteralNode);
+					
+					// Set A to the literal value, B to 1.0 (identity multiplication)
+					UEdGraphPin* APin = LiteralNode->FindPin(TEXT("A"));
+					UEdGraphPin* BPin = LiteralNode->FindPin(TEXT("B"));
+					if (APin)
+					{
+						APin->DefaultValue = LiteralValueStr;
+					}
+					if (BPin)
+					{
+						BPin->DefaultValue = TEXT("1.0");
+					}
+					
+					// Cache the node so we can reuse it
+					Ctx.TempIdToNode.Add(LiteralNodeKey, LiteralNode);
+					
+					return FindOutputPin(LiteralNode, TEXT("ReturnValue"));
+				}
+				else
+				{
+					Ctx.Warnings.Add(FString::Printf(TEXT("Could not create literal node for '%s'"), *SymName));
+				}
+				return nullptr;
+			}
+			
+			// Check if this is a string literal
+			if (NodeId.StartsWith(TEXT("_literalstr_")))
+			{
+				FString LiteralValueStr = PinName; // PinName holds the literal string
+				
+				// Check if we already created a node for this literal
+				FString LiteralNodeKey = TEXT("_literalstrnode_") + SymName;
+				if (UEdGraphNode** ExistingNode = Ctx.TempIdToNode.Find(LiteralNodeKey))
+				{
+					return FindOutputPin(*ExistingNode, TEXT("ReturnValue"));
+				}
+				
+				// For strings, use Concat_StrStr with empty string as identity
+				UFunction* ConcatFunc = UKismetStringLibrary::StaticClass()->FindFunctionByName(TEXT("Concat_StrStr"));
+				if (ConcatFunc)
+				{
+					UK2Node_CallFunction* LiteralNode = NewObject<UK2Node_CallFunction>(Ctx.Graph);
+					LiteralNode->SetFromFunction(ConcatFunc);
+					LiteralNode->NodePosX = Ctx.CurrentX;
+					LiteralNode->NodePosY = Ctx.CurrentY;
+					Ctx.Graph->AddNode(LiteralNode, false, false);
+					LiteralNode->AllocateDefaultPins();
+					EnsureNodeHasValidGuid(LiteralNode);
+					
+					// Set A to the literal string, B to empty string
+					UEdGraphPin* APin = LiteralNode->FindPin(TEXT("A"));
+					UEdGraphPin* BPin = LiteralNode->FindPin(TEXT("B"));
+					if (APin)
+					{
+						APin->DefaultValue = LiteralValueStr;
+					}
+					if (BPin)
+					{
+						BPin->DefaultValue = TEXT("");
+					}
+					
+					// Cache the node
+					Ctx.TempIdToNode.Add(LiteralNodeKey, LiteralNode);
+					
+					return FindOutputPin(LiteralNode, TEXT("ReturnValue"));
+				}
+				return nullptr;
+			}
+			
+			// First try direct lookup by variable name key (most reliable for let bindings)
+			FString VarKey = TEXT("_var_") + SymName;
+			if (UEdGraphNode** FoundNode = Ctx.TempIdToNode.Find(VarKey))
+			{
+				UEdGraphPin* Pin = FindOutputPin(*FoundNode, PinName);
+				if (Pin)
+				{
+					return Pin;
+				}
+			}
+			
+			// Then try direct lookup by NodeId (used for event parameters, etc.)
+			if (UEdGraphNode** FoundNode = Ctx.TempIdToNode.Find(NodeId))
+			{
+				UEdGraphPin* Pin = FindOutputPin(*FoundNode, PinName);
+				if (Pin)
+				{
+					return Pin;
+				}
+			}
+			else
+			{
+				// Node not in TempIdToNode by direct key, try other methods
+				for (auto& Pair : Ctx.TempIdToNode)
+				{
+					Ctx.Warnings.Add(FString::Printf(TEXT("  Key: %s"), *Pair.Key.Left(30)));
+				}
+			}
+			
+			// Fallback: iterate through all TempIdToNode entries looking for matching GUID
+			for (auto& Pair : Ctx.TempIdToNode)
+			{
+				if (Pair.Value && Pair.Value->NodeGuid.ToString() == NodeId)
+				{
+					UEdGraphPin* Pin = FindOutputPin(Pair.Value, PinName);
+					if (Pin)
+					{
+						return Pin;
+					}
+				}
+			}
+			
+			// Last resort: search graph nodes directly
+			for (UEdGraphNode* Node : Ctx.Graph->Nodes)
+			{
+				if (Node && Node->NodeGuid.ToString() == NodeId)
+				{
+					UEdGraphPin* Pin = FindOutputPin(Node, PinName);
+					if (Pin)
+					{
+						// Add to TempIdToNode for future lookups
+						Ctx.TempIdToNode.Add(NodeId, Node);
+						return Pin;
+					}
+				}
+			}
+			
+			Ctx.Warnings.Add(FString::Printf(TEXT("Variable '%s' registered but node/pin not found. NodeId=%s, PinName=%s"), *SymName, *NodeId, *PinName));
+		}
+		
+		// Check for self
+		if (SymName.Equals(TEXT("self"), ESearchCase::IgnoreCase))
+		{
+			UK2Node_Self* SelfNode = NewObject<UK2Node_Self>(Ctx.Graph);
+			SelfNode->NodePosX = Ctx.CurrentX;
+			SelfNode->NodePosY = Ctx.CurrentY;
+			Ctx.Graph->AddNode(SelfNode, false, false);
+			SelfNode->AllocateDefaultPins();
+			EnsureNodeHasValidGuid(SelfNode);
+			
+			FString TempId = Ctx.GenerateTempId();
+			Ctx.TempIdToNode.Add(TempId, SelfNode);
+			
+			return FindOutputPin(SelfNode, TEXT(""));
+		}
+		
+		// Try to create a variable get node
+		UK2Node_VariableGet* VarGet = NewObject<UK2Node_VariableGet>(Ctx.Graph);
+		VarGet->VariableReference.SetSelfMember(FName(*SymName));
+		VarGet->NodePosX = Ctx.CurrentX;
+		VarGet->NodePosY = Ctx.CurrentY;
+		Ctx.Graph->AddNode(VarGet, false, false);
+		VarGet->AllocateDefaultPins();
+		EnsureNodeHasValidGuid(VarGet);
+		
+		FString TempId = Ctx.GenerateTempId();
+		Ctx.TempIdToNode.Add(TempId, VarGet);
+		
+		return FindOutputPin(VarGet, TEXT(""));
+	}
+	
+	// Number literal - need to create a constant node or set as default value
+	if (Expr->IsNumber())
+	{
+		// For now, return nullptr - the caller should set the default value
+		return nullptr;
+	}
+	
+	// String literal
+	if (Expr->IsString())
+	{
+		return nullptr;
+	}
+	
+	// List - could be a function call or operator
+	if (Expr->IsList() && Expr->Num() > 0)
+	{
+		FString FormName = Expr->GetFormName();
+		
+		// Handle (valid? obj) or (IsValid obj) - null/validity check
+		if ((FormName.Equals(TEXT("valid?"), ESearchCase::IgnoreCase) || 
+			 FormName.Equals(TEXT("IsValid"), ESearchCase::IgnoreCase)) && Expr->Num() >= 2)
+		{
+			UFunction* IsValidFunc = UKismetSystemLibrary::StaticClass()->FindFunctionByName(TEXT("IsValid"));
+			if (IsValidFunc)
+			{
+				UK2Node_CallFunction* IsValidNode = NewObject<UK2Node_CallFunction>(Ctx.Graph);
+				IsValidNode->SetFromFunction(IsValidFunc);
+				IsValidNode->NodePosX = Ctx.CurrentX;
+				IsValidNode->NodePosY = Ctx.CurrentY;
+				Ctx.Graph->AddNode(IsValidNode, false, false);
+				IsValidNode->AllocateDefaultPins();
+				EnsureNodeHasValidGuid(IsValidNode);
+				
+				// Connect the object to check
+				UEdGraphPin* ObjectPin = IsValidNode->FindPin(TEXT("Object"));
+				if (ObjectPin)
+				{
+					UEdGraphPin* ObjSource = ResolveLispExpression(Expr->Get(1), Ctx);
+					if (ObjSource)
+					{
+						ConnectNodes(ObjSource, ObjectPin, Ctx);
+					}
+				}
+				
+				FString TempId = Ctx.GenerateTempId();
+				Ctx.TempIdToNode.Add(TempId, IsValidNode);
+				return FindOutputPin(IsValidNode, TEXT("ReturnValue"));
+			}
+		}
+		
+		// Handle (get-actor-of-class "ClassName") - GetActorOfClass
+		if ((FormName.Equals(TEXT("get-actor-of-class"), ESearchCase::IgnoreCase) || 
+			 FormName.Equals(TEXT("GetActorOfClass"), ESearchCase::IgnoreCase)) && Expr->Num() >= 2)
+		{
+			UFunction* GetActorFunc = UGameplayStatics::StaticClass()->FindFunctionByName(TEXT("GetActorOfClass"));
+			if (GetActorFunc)
+			{
+				UK2Node_CallFunction* GetActorNode = NewObject<UK2Node_CallFunction>(Ctx.Graph);
+				GetActorNode->SetFromFunction(GetActorFunc);
+				GetActorNode->NodePosX = Ctx.CurrentX;
+				GetActorNode->NodePosY = Ctx.CurrentY;
+				Ctx.Graph->AddNode(GetActorNode, false, false);
+				GetActorNode->AllocateDefaultPins();
+				EnsureNodeHasValidGuid(GetActorNode);
+				
+				// Set the ActorClass pin
+				FString ClassName = Expr->Get(1)->IsString() ? Expr->Get(1)->StringValue : 
+								   (Expr->Get(1)->IsSymbol() ? Expr->Get(1)->StringValue : TEXT(""));
+				
+				UEdGraphPin* ClassPin = GetActorNode->FindPin(TEXT("ActorClass"));
+				if (ClassPin)
+				{
+					// Try to find the class
+					UClass* ActorClass = FindObject<UClass>(nullptr, *FString::Printf(TEXT("/Script/Engine.%s"), *ClassName));
+					if (!ActorClass)
+					{
+						ActorClass = FindObject<UClass>(nullptr, *ClassName);
+					}
+					if (ActorClass)
+					{
+						ClassPin->DefaultObject = ActorClass;
+					}
+				}
+				
+				// Set WorldContextObject to self
+				UEdGraphPin* WorldPin = GetActorNode->FindPin(TEXT("WorldContextObject"));
+				if (WorldPin)
+				{
+					// Create Self node
+					UK2Node_Self* SelfNode = NewObject<UK2Node_Self>(Ctx.Graph);
+					SelfNode->NodePosX = Ctx.CurrentX - 150;
+					SelfNode->NodePosY = Ctx.CurrentY;
+					Ctx.Graph->AddNode(SelfNode, false, false);
+					SelfNode->AllocateDefaultPins();
+					
+					UEdGraphPin* SelfOutput = SelfNode->FindPin(UEdGraphSchema_K2::PN_Self);
+					if (SelfOutput)
+					{
+						ConnectNodes(SelfOutput, WorldPin, Ctx);
+					}
+				}
+				
+				FString TempId = Ctx.GenerateTempId();
+				Ctx.TempIdToNode.Add(TempId, GetActorNode);
+				return FindOutputPin(GetActorNode, TEXT("ReturnValue"));
+			}
+		}
+		
+		// Handle (get-all-actors-of-class "ClassName") - GetAllActorsOfClass
+		if ((FormName.Equals(TEXT("get-all-actors-of-class"), ESearchCase::IgnoreCase) || 
+			 FormName.Equals(TEXT("GetAllActorsOfClass"), ESearchCase::IgnoreCase)) && Expr->Num() >= 2)
+		{
+			UFunction* GetAllActorsFunc = UGameplayStatics::StaticClass()->FindFunctionByName(TEXT("GetAllActorsOfClass"));
+			if (GetAllActorsFunc)
+			{
+				UK2Node_CallFunction* GetAllNode = NewObject<UK2Node_CallFunction>(Ctx.Graph);
+				GetAllNode->SetFromFunction(GetAllActorsFunc);
+				GetAllNode->NodePosX = Ctx.CurrentX;
+				GetAllNode->NodePosY = Ctx.CurrentY;
+				Ctx.Graph->AddNode(GetAllNode, false, false);
+				GetAllNode->AllocateDefaultPins();
+				EnsureNodeHasValidGuid(GetAllNode);
+				
+				// Set the ActorClass pin
+				FString ClassName = Expr->Get(1)->IsString() ? Expr->Get(1)->StringValue : 
+								   (Expr->Get(1)->IsSymbol() ? Expr->Get(1)->StringValue : TEXT(""));
+				
+				UEdGraphPin* ClassPin = GetAllNode->FindPin(TEXT("ActorClass"));
+				if (ClassPin)
+				{
+					UClass* ActorClass = FindObject<UClass>(nullptr, *FString::Printf(TEXT("/Script/Engine.%s"), *ClassName));
+					if (!ActorClass)
+					{
+						ActorClass = FindObject<UClass>(nullptr, *ClassName);
+					}
+					if (ActorClass)
+					{
+						ClassPin->DefaultObject = ActorClass;
+					}
+				}
+				
+				// Set WorldContextObject to self
+				UEdGraphPin* WorldPin = GetAllNode->FindPin(TEXT("WorldContextObject"));
+				if (WorldPin)
+				{
+					UK2Node_Self* SelfNode = NewObject<UK2Node_Self>(Ctx.Graph);
+					SelfNode->NodePosX = Ctx.CurrentX - 150;
+					SelfNode->NodePosY = Ctx.CurrentY;
+					Ctx.Graph->AddNode(SelfNode, false, false);
+					SelfNode->AllocateDefaultPins();
+					
+					UEdGraphPin* SelfOutput = SelfNode->FindPin(UEdGraphSchema_K2::PN_Self);
+					if (SelfOutput)
+					{
+						ConnectNodes(SelfOutput, WorldPin, Ctx);
+					}
+				}
+				
+				FString TempId = Ctx.GenerateTempId();
+				Ctx.TempIdToNode.Add(TempId, GetAllNode);
+				return FindOutputPin(GetAllNode, TEXT("OutActors"));
+			}
+		}
+		
+		// Handle (select condition true-val false-val) - inline conditional (like ternary operator)
+		if (FormName.Equals(TEXT("select"), ESearchCase::IgnoreCase) && Expr->Num() >= 4)
+		{
+			UFunction* SelectFunc = UKismetMathLibrary::StaticClass()->FindFunctionByName(TEXT("SelectFloat"));
+			if (SelectFunc)
+			{
+				UK2Node_CallFunction* SelectNode = NewObject<UK2Node_CallFunction>(Ctx.Graph);
+				SelectNode->SetFromFunction(SelectFunc);
+				SelectNode->NodePosX = Ctx.CurrentX;
+				SelectNode->NodePosY = Ctx.CurrentY;
+				Ctx.Graph->AddNode(SelectNode, false, false);
+				SelectNode->AllocateDefaultPins();
+				EnsureNodeHasValidGuid(SelectNode);
+				
+				// Connect condition (bPickA - when true, A is selected)
+				// Try both "bPickA" and "PickA" as the pin name varies
+				UEdGraphPin* ConditionPin = SelectNode->FindPin(TEXT("bPickA"));
+				if (!ConditionPin)
+				{
+					ConditionPin = SelectNode->FindPin(TEXT("PickA"));
+				}
+				if (!ConditionPin)
+				{
+					// Fall back to finding any boolean input pin
+					for (UEdGraphPin* Pin : SelectNode->Pins)
+					{
+						if (Pin->Direction == EGPD_Input && 
+							Pin->PinType.PinCategory == UEdGraphSchema_K2::PC_Boolean)
+						{
+							ConditionPin = Pin;
+							break;
+						}
+					}
+				}
+				if (ConditionPin)
+				{
+					UEdGraphPin* CondSource = ResolveLispExpression(Expr->Get(1), Ctx);
+					if (CondSource)
+					{
+						ConnectNodes(CondSource, ConditionPin, Ctx);
+					}
+					else
+					{
+						Ctx.Warnings.Add(FString::Printf(TEXT("select: Could not resolve condition expression")));
+					}
+				}
+				else
+				{
+					Ctx.Warnings.Add(TEXT("select: Could not find condition pin on SelectFloat node"));
+				}
+				
+				// Connect A (true value)
+				UEdGraphPin* APin = SelectNode->FindPin(TEXT("A"));
+				if (APin)
+				{
+					FLispNodePtr TrueExpr = Expr->Get(2);
+					if (TrueExpr->IsList() || TrueExpr->IsSymbol())
+					{
+						UEdGraphPin* TrueSource = ResolveLispExpression(TrueExpr, Ctx);
+						if (TrueSource) ConnectNodes(TrueSource, APin, Ctx);
+					}
+					else if (TrueExpr->IsNumber())
+					{
+						SetNumericPinDefaultValue(APin, TrueExpr->NumberValue);
+					}
+				}
+				
+				// Connect B (false value)
+				UEdGraphPin* BPin = SelectNode->FindPin(TEXT("B"));
+				if (BPin)
+				{
+					FLispNodePtr FalseExpr = Expr->Get(3);
+					if (FalseExpr->IsList() || FalseExpr->IsSymbol())
+					{
+						UEdGraphPin* FalseSource = ResolveLispExpression(FalseExpr, Ctx);
+						if (FalseSource) ConnectNodes(FalseSource, BPin, Ctx);
+					}
+					else if (FalseExpr->IsNumber())
+					{
+						SetNumericPinDefaultValue(BPin, FalseExpr->NumberValue);
+					}
+				}
+				
+				FString TempId = Ctx.GenerateTempId();
+				Ctx.TempIdToNode.Add(TempId, SelectNode);
+				return FindOutputPin(SelectNode, TEXT("ReturnValue"));
+			}
+		}
+		
+		// Handle (component Name) - get component by name
+		if (FormName.Equals(TEXT("component"), ESearchCase::IgnoreCase) && Expr->Num() >= 2)
+		{
+			FString CompName = Expr->Get(1)->IsSymbol() ? Expr->Get(1)->StringValue : 
+							   Expr->Get(1)->IsString() ? Expr->Get(1)->StringValue : TEXT("");
+			
+			// Create a GetComponentByClass or reference to component variable
+			UK2Node_VariableGet* VarGet = NewObject<UK2Node_VariableGet>(Ctx.Graph);
+			VarGet->VariableReference.SetSelfMember(FName(*CompName));
+			VarGet->NodePosX = Ctx.CurrentX;
+			VarGet->NodePosY = Ctx.CurrentY;
+			Ctx.Graph->AddNode(VarGet, false, false);
+			VarGet->AllocateDefaultPins();
+			EnsureNodeHasValidGuid(VarGet);
+			
+			FString TempId = Ctx.GenerateTempId();
+			Ctx.TempIdToNode.Add(TempId, VarGet);
+			
+			// Find the output pin (component variable)
+			for (UEdGraphPin* Pin : VarGet->Pins)
+			{
+				if (Pin->Direction == EGPD_Output && Pin->PinName == FName(*CompName))
+				{
+					return Pin;
+				}
+			}
+			return FindOutputPin(VarGet, TEXT(""));
+		}
+		
+		// Handle (make Type :field val ...) - struct construction with named fields
+		if (FormName.Equals(TEXT("make"), ESearchCase::IgnoreCase) && Expr->Num() >= 2)
+		{
+			FString TypeName = Expr->Get(1)->IsSymbol() ? Expr->Get(1)->StringValue : TEXT("");
+			
+			// Map type names to Make functions
+			FString MakeFuncName;
+			if (TypeName.Equals(TEXT("Rotator"), ESearchCase::IgnoreCase))
+			{
+				MakeFuncName = TEXT("MakeRotator");
+			}
+			else if (TypeName.Equals(TEXT("Vector"), ESearchCase::IgnoreCase))
+			{
+				MakeFuncName = TEXT("MakeVector");
+			}
+			else if (TypeName.Equals(TEXT("Transform"), ESearchCase::IgnoreCase))
+			{
+				MakeFuncName = TEXT("MakeTransform");
+			}
+			else if (TypeName.Equals(TEXT("Color"), ESearchCase::IgnoreCase) || TypeName.Equals(TEXT("LinearColor"), ESearchCase::IgnoreCase))
+			{
+				MakeFuncName = TEXT("MakeColor");
+			}
+			
+			if (!MakeFuncName.IsEmpty())
+			{
+				UFunction* MakeFunc = UKismetMathLibrary::StaticClass()->FindFunctionByName(*MakeFuncName);
+				if (MakeFunc)
+				{
+					UK2Node_CallFunction* MakeNode = NewObject<UK2Node_CallFunction>(Ctx.Graph);
+					MakeNode->SetFromFunction(MakeFunc);
+					MakeNode->NodePosX = Ctx.CurrentX;
+					MakeNode->NodePosY = Ctx.CurrentY;
+					Ctx.Graph->AddNode(MakeNode, false, false);
+					MakeNode->AllocateDefaultPins();
+					EnsureNodeHasValidGuid(MakeNode);
+					
+					// Process keyword arguments: :Pitch 0 :Yaw 90 :Roll 0
+					for (int32 i = 2; i < Expr->Num(); i++)
+					{
+						if (Expr->Get(i)->IsKeyword())
+						{
+							FString KeyName = Expr->Get(i)->StringValue;
+							// Remove leading colon if present
+							if (KeyName.StartsWith(TEXT(":")))
+							{
+								KeyName = KeyName.Mid(1);
+							}
+							
+							if (i + 1 < Expr->Num())
+							{
+								UEdGraphPin* Pin = MakeNode->FindPin(FName(*KeyName));
+								if (Pin)
+								{
+									FLispNodePtr ValExpr = Expr->Get(i + 1);
+									if (ValExpr->IsList() || ValExpr->IsSymbol())
+									{
+										UEdGraphPin* ValSource = ResolveLispExpression(ValExpr, Ctx);
+										if (ValSource) ConnectNodes(ValSource, Pin, Ctx);
+									}
+									else if (ValExpr->IsNumber())
+									{
+										SetNumericPinDefaultValue(Pin, ValExpr->NumberValue);
+									}
+								}
+								i++; // Skip the value
+							}
+						}
+					}
+					
+					FString TempId = Ctx.GenerateTempId();
+					Ctx.TempIdToNode.Add(TempId, MakeNode);
+					return FindOutputPin(MakeNode, TEXT("ReturnValue"));
+				}
+			}
+		}
+		
+		// Handle (get VarName) - explicit variable get (alternative to just using VarName directly)
+		if (FormName.Equals(TEXT("get"), ESearchCase::IgnoreCase) && Expr->Num() >= 2)
+		{
+			FString VarName = Expr->Get(1)->IsSymbol() ? Expr->Get(1)->StringValue : TEXT("");
+			
+			// Check if it's a known variable from let bindings
+			if (Ctx.VariableToNodeId.Contains(VarName))
+			{
+				FString NodeId = Ctx.VariableToNodeId[VarName];
+				FString PinName = Ctx.VariableToPin.Contains(VarName) ? Ctx.VariableToPin[VarName] : TEXT("");
+				
+				// First try direct lookup by NodeId (key is NodeGuid string)
+				if (UEdGraphNode** FoundNode = Ctx.TempIdToNode.Find(NodeId))
+				{
+					UEdGraphPin* Pin = FindOutputPin(*FoundNode, PinName);
+					if (Pin)
+					{
+						return Pin;
+					}
+				}
+				
+				// Fallback: search all nodes by GUID
+				for (auto& Pair : Ctx.TempIdToNode)
+				{
+					if (Pair.Value && Pair.Value->NodeGuid.ToString() == NodeId)
+					{
+						UEdGraphPin* Pin = FindOutputPin(Pair.Value, PinName);
+						if (Pin)
+						{
+							return Pin;
+						}
+					}
+				}
+				
+				// Last resort: search graph nodes directly
+				for (UEdGraphNode* Node : Ctx.Graph->Nodes)
+				{
+					if (Node && Node->NodeGuid.ToString() == NodeId)
+					{
+						UEdGraphPin* Pin = FindOutputPin(Node, PinName);
+						if (Pin)
+						{
+							// Add to TempIdToNode for future lookups
+							Ctx.TempIdToNode.Add(NodeId, Node);
+							return Pin;
+						}
+					}
+				}
+				
+				Ctx.Warnings.Add(FString::Printf(TEXT("(get %s): Variable registered but could not find node/pin. NodeId=%s, PinName=%s"), *VarName, *NodeId, *PinName));
+			}
+			
+			// Fall back to creating a Blueprint variable get node (for class member variables)
+			UK2Node_VariableGet* VarGet = NewObject<UK2Node_VariableGet>(Ctx.Graph);
+			VarGet->VariableReference.SetSelfMember(FName(*VarName));
+			VarGet->NodePosX = Ctx.CurrentX;
+			VarGet->NodePosY = Ctx.CurrentY;
+			Ctx.Graph->AddNode(VarGet, false, false);
+			VarGet->AllocateDefaultPins();
+			EnsureNodeHasValidGuid(VarGet);
+			
+			FString TempId = Ctx.GenerateTempId();
+			Ctx.TempIdToNode.Add(TempId, VarGet);
+			
+			return FindOutputPin(VarGet, TEXT(""));
+		}
+		
+		// Handle (. struct field) - struct field access / break struct
+		// Example: (. loc X) gets the X component of a vector
+		if (FormName == TEXT(".") && Expr->Num() >= 3)
+		{
+			FString FieldName = Expr->Get(2)->IsSymbol() ? Expr->Get(2)->StringValue : TEXT("");
+			FString StructExprStr = Expr->Get(1)->ToString(false);
+			
+			// Resolve the struct expression
+			UEdGraphPin* StructPin = ResolveLispExpression(Expr->Get(1), Ctx);
+			
+			if (StructPin)
+			{
+				// Check if this is a vector - use BreakVector
+				if (StructPin->PinType.PinCategory == UEdGraphSchema_K2::PC_Struct)
+				{
+					UScriptStruct* StructType = Cast<UScriptStruct>(StructPin->PinType.PinSubCategoryObject.Get());
+					if (StructType)
+					{
+						FString StructName = StructType->GetName();
+						
+						// Find appropriate break function
+						FString BreakFuncName;
+						if (StructName == TEXT("Vector"))
+						{
+							BreakFuncName = TEXT("BreakVector");
+						}
+						else if (StructName == TEXT("Rotator"))
+						{
+							BreakFuncName = TEXT("BreakRotator");
+						}
+						else if (StructName == TEXT("Vector2D"))
+						{
+							BreakFuncName = TEXT("BreakVector2D");
+						}
+						else if (StructName == TEXT("Transform"))
+						{
+							BreakFuncName = TEXT("BreakTransform");
+						}
+						else if (StructName == TEXT("LinearColor") || StructName == TEXT("Color"))
+						{
+							BreakFuncName = TEXT("BreakColor");
+						}
+						
+						if (!BreakFuncName.IsEmpty())
+						{
+							// Create a cache key based on the source pin's node GUID and pin name
+							// This allows reusing the same BreakVector for multiple field accesses
+							UEdGraphNode* SourceNode = StructPin->GetOwningNode();
+							FString BreakNodeCacheKey = FString::Printf(TEXT("_break_%s_%s_%s"), 
+								*SourceNode->NodeGuid.ToString(), 
+								*StructPin->PinName.ToString(),
+								*BreakFuncName);
+							
+							// Check if we already have a Break node for this struct
+							UK2Node_CallFunction* BreakNode = nullptr;
+							if (UEdGraphNode** ExistingNode = Ctx.TempIdToNode.Find(BreakNodeCacheKey))
+							{
+								BreakNode = Cast<UK2Node_CallFunction>(*ExistingNode);
+							}
+							
+							// Create new Break node if we don't have one cached
+							if (!BreakNode)
+							{
+								UFunction* BreakFunc = UKismetMathLibrary::StaticClass()->FindFunctionByName(*BreakFuncName);
+								if (BreakFunc)
+								{
+									BreakNode = NewObject<UK2Node_CallFunction>(Ctx.Graph);
+									BreakNode->SetFromFunction(BreakFunc);
+									BreakNode->NodePosX = Ctx.CurrentX;
+									BreakNode->NodePosY = Ctx.CurrentY;
+									Ctx.Graph->AddNode(BreakNode, false, false);
+									BreakNode->AllocateDefaultPins();
+									
+									// Ensure the node has a valid GUID for variable tracking
+									EnsureNodeHasValidGuid(BreakNode);
+									
+									// Connect the struct to the Break node's input
+									// Look for the input pin that matches our struct type, NOT the self pin
+									UEdGraphPin* StructInputPin = nullptr;
+									for (UEdGraphPin* Pin : BreakNode->Pins)
+									{
+										if (Pin->Direction == EGPD_Input && 
+											Pin->PinType.PinCategory == UEdGraphSchema_K2::PC_Struct &&
+											!Pin->bHidden)
+										{
+											StructInputPin = Pin;
+											break;
+										}
+									}
+									
+									// If no struct input found, try first non-exec, non-self input
+									if (!StructInputPin)
+									{
+										for (UEdGraphPin* Pin : BreakNode->Pins)
+										{
+											if (Pin->Direction == EGPD_Input && 
+												Pin->PinType.PinCategory != UEdGraphSchema_K2::PC_Exec &&
+												Pin->PinName != UEdGraphSchema_K2::PN_Self &&
+												!Pin->bHidden)
+											{
+												StructInputPin = Pin;
+												break;
+											}
+										}
+									}
+									
+									if (StructInputPin)
+									{
+										ConnectNodes(StructPin, StructInputPin, Ctx);
+									}
+									else
+									{
+										Ctx.Warnings.Add(FString::Printf(TEXT("Could not find struct input pin on %s"), *BreakFuncName));
+									}
+									
+									// Cache the Break node for reuse - use FindOrAdd to avoid overwriting
+									// Also check if this key might conflict with any variable bindings
+									// Cache the BreakNode for reuse (avoid duplicates)
+									if (!Ctx.TempIdToNode.Contains(BreakNodeCacheKey))
+									{
+										Ctx.TempIdToNode.Add(BreakNodeCacheKey, BreakNode);
+									}
+								}
+							}
+							
+							// Find output pin matching field name
+							if (BreakNode)
+							{
+								for (UEdGraphPin* Pin : BreakNode->Pins)
+								{
+									if (Pin->Direction == EGPD_Output && 
+										Pin->PinName.ToString().Equals(FieldName, ESearchCase::IgnoreCase))
+									{
+										return Pin;
+									}
+								}
+								
+								Ctx.Warnings.Add(FString::Printf(TEXT("Field '%s' not found on %s. Available: X, Y, Z"), *FieldName, *StructName));
+							}
+						}
+					}
+				}
+			}
+			
+			Ctx.Warnings.Add(FString::Printf(TEXT("Could not resolve struct field access: (. %s %s)"), 
+				*StructExprStr, *FieldName));
+			return nullptr;
+		}
+		
+		// Handle (asset "path/to/asset") - this is a special form that returns nullptr
+		// The caller should detect this and set Pin->DefaultObject directly
+		if (FormName.Equals(TEXT("asset"), ESearchCase::IgnoreCase) && Expr->Num() >= 2)
+		{
+			// We store the asset path in the context for the caller to use
+			FString AssetPath = Expr->Get(1)->IsString() ? Expr->Get(1)->StringValue : Expr->Get(1)->ToString(false);
+			
+			// Return nullptr - the caller must handle asset assignment by setting DefaultObject
+			// We'll add a helper context field to pass this info
+			Ctx.LastAssetPath = AssetPath;
+			return nullptr;
+		}
+		
+		// Map Lisp operators to UE math function names
+		// UE5 uses DoubleDouble for float math, IntInt for integers
+		static TMap<FString, FString> OpToFuncMap;
+		if (OpToFuncMap.Num() == 0)
+		{
+			// Basic arithmetic - using DoubleDouble for float compatibility
+			OpToFuncMap.Add(TEXT("+"), TEXT("Add_DoubleDouble"));
+			OpToFuncMap.Add(TEXT("-"), TEXT("Subtract_DoubleDouble"));
+			OpToFuncMap.Add(TEXT("*"), TEXT("Multiply_DoubleDouble"));
+			OpToFuncMap.Add(TEXT("/"), TEXT("Divide_DoubleDouble"));
+			OpToFuncMap.Add(TEXT("%"), TEXT("Percent_FloatFloat"));
+			OpToFuncMap.Add(TEXT("mod"), TEXT("Percent_FloatFloat"));
+			OpToFuncMap.Add(TEXT("fmod"), TEXT("Percent_FloatFloat"));  // Use Percent for consistent float return
+			OpToFuncMap.Add(TEXT("remainder"), TEXT("FMod"));  // FMod for when you need both quotient and remainder
+			
+			// Comparisons - use DoubleDouble for float compatibility
+			OpToFuncMap.Add(TEXT("<"), TEXT("Less_DoubleDouble"));
+			OpToFuncMap.Add(TEXT(">"), TEXT("Greater_DoubleDouble"));
+			OpToFuncMap.Add(TEXT("<="), TEXT("LessEqual_DoubleDouble"));
+			OpToFuncMap.Add(TEXT(">="), TEXT("GreaterEqual_DoubleDouble"));
+			OpToFuncMap.Add(TEXT("=="), TEXT("EqualEqual_DoubleDouble"));
+			OpToFuncMap.Add(TEXT("="), TEXT("EqualEqual_DoubleDouble"));
+			OpToFuncMap.Add(TEXT("!="), TEXT("NotEqual_DoubleDouble"));
+			
+			// Boolean operations
+			OpToFuncMap.Add(TEXT("and"), TEXT("BooleanAND"));
+			OpToFuncMap.Add(TEXT("or"), TEXT("BooleanOR"));
+			OpToFuncMap.Add(TEXT("not"), TEXT("Not_PreBool"));
+		}
+		
+		// Handle math operators: (+ a b), (* a b), etc.
+		if (FString* BaseFuncName = OpToFuncMap.Find(FormName))
+		{
+			// Try to determine the type of the first argument to pick the right operator
+			FString FuncName = *BaseFuncName;
+			bool bIsIntegerOp = false;
+			
+			// Check first argument's type
+			if (Expr->Num() > 1)
+			{
+				FLispNodePtr FirstArg = Expr->Get(1);
+				if (FirstArg->IsSymbol())
+				{
+					// Check if this is a known variable
+					FString SymName = FirstArg->StringValue;
+					if (Ctx.VariableToNodeId.Contains(SymName))
+					{
+						FString NodeId = Ctx.VariableToNodeId[SymName];
+						// Check if it's an integer literal
+						if (NodeId.StartsWith(TEXT("_literalint_")))
+						{
+							bIsIntegerOp = true;
+						}
+						// Check if it's a float literal
+						else if (NodeId.StartsWith(TEXT("_literal_")))
+						{
+							// Float literals are not integers
+							bIsIntegerOp = false;
+						}
+						else
+						{
+							// Check the actual pin type - try _var_ prefixed key first
+							FString VarKey = TEXT("_var_") + SymName;
+							UEdGraphNode** FoundNode = Ctx.TempIdToNode.Find(VarKey);
+							if (!FoundNode)
+							{
+								// Fall back to NodeId (GUID) lookup
+								FoundNode = Ctx.TempIdToNode.Find(NodeId);
+							}
+							
+							if (FoundNode)
+							{
+								FString PinName = Ctx.VariableToPin.Contains(SymName) ? Ctx.VariableToPin[SymName] : TEXT("");
+								UEdGraphPin* Pin = FindOutputPin(*FoundNode, PinName);
+								if (Pin)
+								{
+									// Only set to integer if it's actually an integer type
+									if (Pin->PinType.PinCategory == UEdGraphSchema_K2::PC_Int)
+									{
+										bIsIntegerOp = true;
+									}
+									// Explicitly NOT an integer if it's a float/double
+									else if (Pin->PinType.PinCategory == UEdGraphSchema_K2::PC_Real ||
+											 Pin->PinType.PinCategory == UEdGraphSchema_K2::PC_Float ||
+											 Pin->PinType.PinCategory == UEdGraphSchema_K2::PC_Double)
+									{
+										bIsIntegerOp = false;
+									}
+								}
+							}
+						}
+					}
+					// Check if it's a variable get node we can inspect
+					else
+					{
+						// Try to resolve and check type
+						UEdGraphPin* ArgPin = ResolveLispExpression(FirstArg, Ctx);
+						if (ArgPin && ArgPin->PinType.PinCategory == UEdGraphSchema_K2::PC_Int)
+						{
+							bIsIntegerOp = true;
+						}
+					}
+				}
+				else if (FirstArg->IsNumber())
+				{
+					// Check if the number is an integer (no decimal point/exponent in value OR string representation)
+					double val = FirstArg->NumberValue;
+					bool bIsWholeNumber = (val == FMath::FloorToDouble(val) && FMath::Abs(val) < 2147483647);
+					bool bLooksLikeFloat = false;
+					if (!FirstArg->StringValue.IsEmpty())
+					{
+						bLooksLikeFloat = FirstArg->StringValue.Contains(TEXT(".")) || 
+										  FirstArg->StringValue.Contains(TEXT("e")) || 
+										  FirstArg->StringValue.Contains(TEXT("E"));
+					}
+					
+					// Only treat as integer if it's a whole number AND wasn't written as a float
+					if (bIsWholeNumber && !bLooksLikeFloat)
+					{
+						bIsIntegerOp = true;
+					}
+				}
+			}
+			
+			// Also check all other arguments - if ANY argument looks like a float, use float
+			for (int32 i = 1; i < Expr->Num(); i++)  // Check ALL args including first
+			{
+				FLispNodePtr Arg = Expr->Get(i);
+				if (Arg->IsNumber())
+				{
+					double val = Arg->NumberValue;
+					// If number has decimal part OR is written as a float (check string representation)
+					// Note: The parser stores the original string in StringValue for numbers too
+					bool bLooksLikeFloat = (val != FMath::FloorToDouble(val));
+					
+					// Also check if the original representation had a decimal point or exponent
+					// Numbers like 24.0, 1e5, 2.5E-3 should be treated as floats
+					if (!bLooksLikeFloat && !Arg->StringValue.IsEmpty())
+					{
+						bLooksLikeFloat = Arg->StringValue.Contains(TEXT(".")) || 
+										  Arg->StringValue.Contains(TEXT("e")) || 
+										  Arg->StringValue.Contains(TEXT("E"));
+					}
+					
+					if (bLooksLikeFloat)
+					{
+						bIsIntegerOp = false;
+						break;
+					}
+				}
+				else if (Arg->IsSymbol())
+				{
+					// Check if this variable is a float type
+					FString SymName = Arg->StringValue;
+					if (Ctx.VariableToNodeId.Contains(SymName))
+					{
+						FString NodeId = Ctx.VariableToNodeId[SymName];
+						// Float literals are not integers
+						if (NodeId.StartsWith(TEXT("_literal_")) && !NodeId.StartsWith(TEXT("_literalint_")))
+						{
+							bIsIntegerOp = false;
+							break;
+						}
+						// Check the actual pin type
+						FString VarKey = TEXT("_var_") + SymName;
+						UEdGraphNode** FoundNode = Ctx.TempIdToNode.Find(VarKey);
+						if (!FoundNode)
+						{
+							FoundNode = Ctx.TempIdToNode.Find(NodeId);
+						}
+						if (FoundNode)
+						{
+							FString PinName = Ctx.VariableToPin.Contains(SymName) ? Ctx.VariableToPin[SymName] : TEXT("");
+							UEdGraphPin* Pin = FindOutputPin(*FoundNode, PinName);
+							if (Pin && (Pin->PinType.PinCategory == UEdGraphSchema_K2::PC_Real ||
+							            Pin->PinType.PinCategory == UEdGraphSchema_K2::PC_Float ||
+							            Pin->PinType.PinCategory == UEdGraphSchema_K2::PC_Double))
+							{
+								bIsIntegerOp = false;
+								break;
+							}
+						}
+					}
+					else
+					{
+						// Unknown variable - try to resolve it and check type
+						// This handles Blueprint variables that aren't in our tracking yet
+						UEdGraphPin* ArgPin = ResolveLispExpression(Arg, Ctx);
+						if (ArgPin && (ArgPin->PinType.PinCategory == UEdGraphSchema_K2::PC_Real ||
+						               ArgPin->PinType.PinCategory == UEdGraphSchema_K2::PC_Float ||
+						               ArgPin->PinType.PinCategory == UEdGraphSchema_K2::PC_Double))
+						{
+							bIsIntegerOp = false;
+							break;
+						}
+					}
+				}
+				else if (Arg->IsList())
+				{
+					// Nested expression - check its operator
+					// If it's a float operation, this should be float too
+					FString NestedFormName = Arg->GetFormName();
+					if (NestedFormName == TEXT("/") || NestedFormName == TEXT("sin") || 
+						NestedFormName == TEXT("cos") || NestedFormName == TEXT("sqrt") ||
+						NestedFormName == TEXT("fmod") || NestedFormName == TEXT("lerp"))
+					{
+						// These always return floats
+						bIsIntegerOp = false;
+						break;
+					}
+				}
+			}
+			
+			// If integer operation, try to use IntInt variant
+			if (bIsIntegerOp)
+			{
+				FString IntFuncName = FuncName.Replace(TEXT("_DoubleDouble"), TEXT("_IntInt"));
+				IntFuncName = IntFuncName.Replace(TEXT("_FloatFloat"), TEXT("_IntInt"));
+				UFunction* IntFunc = UKismetMathLibrary::StaticClass()->FindFunctionByName(*IntFuncName);
+				if (IntFunc)
+				{
+					FuncName = IntFuncName;
+				}
+			}
+			
+			// Special case: Float-only functions need int-to-float conversion
+			bool bNeedsIntToFloatConversion = (FuncName == TEXT("FMod") || 
+											   FuncName == TEXT("Percent_FloatFloat") ||
+											   FuncName.Contains(TEXT("_FloatFloat")) ||
+											   FuncName.Contains(TEXT("_DoubleDouble")));
+			
+			UFunction* Function = UKismetMathLibrary::StaticClass()->FindFunctionByName(*FuncName);
+			if (Function)
+			{
+				UK2Node_CallFunction* MathNode = NewObject<UK2Node_CallFunction>(Ctx.Graph);
+				MathNode->SetFromFunction(Function);
+				MathNode->NodePosX = Ctx.CurrentX;
+				MathNode->NodePosY = Ctx.CurrentY;
+				Ctx.Graph->AddNode(MathNode, false, false);
+				MathNode->AllocateDefaultPins();
+				EnsureNodeHasValidGuid(MathNode);
+				
+				// Connect arguments
+				int32 ArgIndex = 0;
+				for (int32 i = 1; i < Expr->Num(); i++)
+				{
+					int32 PinIndex = 0;
+					for (UEdGraphPin* Pin : MathNode->Pins)
+					{
+						if (Pin->Direction == EGPD_Input && 
+							Pin->PinType.PinCategory != UEdGraphSchema_K2::PC_Exec &&
+							!Pin->bHidden)
+						{
+							if (PinIndex == ArgIndex)
+							{
+								UEdGraphPin* ArgSource = ResolveLispExpression(Expr->Get(i), Ctx);
+								if (ArgSource)
+								{
+									// Check if we need int-to-float conversion (e.g., for FMod)
+									bool bNeedConversion = bNeedsIntToFloatConversion && 
+										ArgSource->PinType.PinCategory == UEdGraphSchema_K2::PC_Int &&
+										(Pin->PinType.PinCategory == UEdGraphSchema_K2::PC_Real ||
+										 Pin->PinType.PinCategory == UEdGraphSchema_K2::PC_Float ||
+										 Pin->PinType.PinCategory == UEdGraphSchema_K2::PC_Double);
+									
+									if (bNeedConversion)
+									{
+										// Use multiply by 1.0 as a reliable int-to-float conversion
+										// This works because UE will auto-convert the int input
+										UFunction* MultiplyFunc = UKismetMathLibrary::StaticClass()->FindFunctionByName(TEXT("Multiply_IntFloat"));
+										if (!MultiplyFunc)
+										{
+											// Fall back to Add_DoubleDouble with 0 - forces type promotion
+											MultiplyFunc = UKismetMathLibrary::StaticClass()->FindFunctionByName(TEXT("Add_DoubleDouble"));
+										}
+										
+										if (MultiplyFunc)
+										{
+											UK2Node_CallFunction* ConvNode = NewObject<UK2Node_CallFunction>(Ctx.Graph);
+											ConvNode->SetFromFunction(MultiplyFunc);
+											ConvNode->NodePosX = Ctx.CurrentX - 150;
+											ConvNode->NodePosY = Ctx.CurrentY + (ArgIndex * 50);
+											Ctx.Graph->AddNode(ConvNode, false, false);
+											ConvNode->AllocateDefaultPins();
+											EnsureNodeHasValidGuid(ConvNode);
+											
+											// Find and connect the integer input (A pin)
+											UEdGraphPin* APin = ConvNode->FindPin(TEXT("A"));
+											if (APin)
+											{
+												ConnectNodes(ArgSource, APin, Ctx);
+											}
+											
+											// Set B to 1.0 (for multiply) or 0.0 (for add)
+											UEdGraphPin* BPin = ConvNode->FindPin(TEXT("B"));
+											if (BPin)
+											{
+												if (MultiplyFunc->GetName().Contains(TEXT("Multiply")))
+												{
+													BPin->DefaultValue = TEXT("1.0");
+												}
+												else
+												{
+													BPin->DefaultValue = TEXT("0.0");
+												}
+											}
+											
+											// Connect conversion output to math node
+											UEdGraphPin* ConvOutput = FindOutputPin(ConvNode, TEXT("ReturnValue"));
+											if (ConvOutput)
+											{
+												ConnectNodes(ConvOutput, Pin, Ctx);
+											}
+										}
+										else
+										{
+											// Last resort: direct connection and hope for implicit conversion
+											ConnectNodes(ArgSource, Pin, Ctx);
+										}
+									}
+									else
+									{
+										ConnectNodes(ArgSource, Pin, Ctx);
+									}
+								}
+								else if (Expr->Get(i)->IsNumber())
+								{
+									double NumVal = Expr->Get(i)->NumberValue;
+									// Check if the pin expects an integer
+									if (Pin->PinType.PinCategory == UEdGraphSchema_K2::PC_Int)
+									{
+										Pin->DefaultValue = FString::Printf(TEXT("%d"), FMath::RoundToInt(NumVal));
+									}
+									else
+									{
+										SetNumericPinDefaultValue(Pin, NumVal);
+									}
+								}
+								ArgIndex++;
+								break;
+							}
+							PinIndex++;
+						}
+					}
+				}
+				
+				FString TempId = Ctx.GenerateTempId();
+				Ctx.TempIdToNode.Add(TempId, MathNode);
+				
+				return FindOutputPin(MathNode, TEXT("ReturnValue"));
+			}
+		}
+		
+// Handle (array-length array) - pure expression returning length
+		if (FormName.Equals(TEXT("array-length"), ESearchCase::IgnoreCase) && Expr->Num() >= 2)
+		{
+			UFunction* LengthFunc = UKismetArrayLibrary::StaticClass()->FindFunctionByName(TEXT("Array_Length"));
+			if (LengthFunc)
+			{
+				UK2Node_CallFunction* LengthNode = NewObject<UK2Node_CallFunction>(Ctx.Graph);
+				LengthNode->SetFromFunction(LengthFunc);
+				LengthNode->NodePosX = Ctx.CurrentX;
+				LengthNode->NodePosY = Ctx.CurrentY;
+				Ctx.Graph->AddNode(LengthNode, false, false);
+				LengthNode->AllocateDefaultPins(); UEdGraphPin* ArrayPin = LengthNode->FindPin(TEXT("TargetArray"));
+				if (ArrayPin)
+				{
+					UEdGraphPin* ArraySource = ResolveLispExpression(Expr->Get(1), Ctx);
+					if (ArraySource)
+					{
+						ConnectNodes(ArraySource, ArrayPin, Ctx);
+					}
+				} FString TempId = Ctx.GenerateTempId();
+				Ctx.TempIdToNode.Add(TempId, LengthNode); return FindOutputPin(LengthNode, TEXT("ReturnValue"));
+			}
+		} // Handle (array-get array index) - get array element
+		if (FormName.Equals(TEXT("array-get"), ESearchCase::IgnoreCase) && Expr->Num() >= 3)
+		{
+			UFunction* GetFunc = UKismetArrayLibrary::StaticClass()->FindFunctionByName(TEXT("Array_Get"));
+			if (GetFunc)
+			{
+				UK2Node_CallFunction* GetNode = NewObject<UK2Node_CallFunction>(Ctx.Graph);
+				GetNode->SetFromFunction(GetFunc);
+				GetNode->NodePosX = Ctx.CurrentX;
+				GetNode->NodePosY = Ctx.CurrentY;
+				Ctx.Graph->AddNode(GetNode, false, false);
+				GetNode->AllocateDefaultPins(); UEdGraphPin* ArrayPin = GetNode->FindPin(TEXT("TargetArray"));
+				if (ArrayPin)
+				{
+					UEdGraphPin* ArraySource = ResolveLispExpression(Expr->Get(1), Ctx);
+					if (ArraySource)
+					{
+						ConnectNodes(ArraySource, ArrayPin, Ctx);
+					}
+				} UEdGraphPin* IndexPin = GetNode->FindPin(TEXT("Index"));
+				if (IndexPin)
+				{
+					if (Expr->Get(2)->IsNumber())
+					{
+						IndexPin->DefaultValue = FString::Printf(TEXT("%d"), (int32)Expr->Get(2)->NumberValue);
+					}
+					else
+					{
+						UEdGraphPin* IndexSource = ResolveLispExpression(Expr->Get(2), Ctx);
+						if (IndexSource)
+						{
+							ConnectNodes(IndexSource, IndexPin, Ctx);
+						}
+					}
+				} FString TempId = Ctx.GenerateTempId();
+				Ctx.TempIdToNode.Add(TempId, GetNode); return FindOutputPin(GetNode, TEXT("Item"));
+			}
+		} // Handle (array-contains array item) - check if array contains item
+		if (FormName.Equals(TEXT("array-contains?"), ESearchCase::IgnoreCase) && Expr->Num() >= 3)
+		{
+			UFunction* ContainsFunc = UKismetArrayLibrary::StaticClass()->FindFunctionByName(TEXT("Array_Contains"));
+			if (ContainsFunc)
+			{
+				UK2Node_CallFunction* ContainsNode = NewObject<UK2Node_CallFunction>(Ctx.Graph);
+				ContainsNode->SetFromFunction(ContainsFunc);
+				ContainsNode->NodePosX = Ctx.CurrentX;
+				ContainsNode->NodePosY = Ctx.CurrentY;
+				Ctx.Graph->AddNode(ContainsNode, false, false);
+				ContainsNode->AllocateDefaultPins(); UEdGraphPin* ArrayPin = ContainsNode->FindPin(TEXT("TargetArray"));
+				if (ArrayPin)
+				{
+					UEdGraphPin* ArraySource = ResolveLispExpression(Expr->Get(1), Ctx);
+					if (ArraySource)
+					{
+						ConnectNodes(ArraySource, ArrayPin, Ctx);
+					}
+				} UEdGraphPin* ItemPin = ContainsNode->FindPin(TEXT("ItemToFind"));
+				if (ItemPin)
+				{
+					UEdGraphPin* ItemSource = ResolveLispExpression(Expr->Get(2), Ctx);
+					if (ItemSource)
+					{
+						ConnectNodes(ItemSource, ItemPin, Ctx);
+					}
+				} FString TempId = Ctx.GenerateTempId();
+				Ctx.TempIdToNode.Add(TempId, ContainsNode); return FindOutputPin(ContainsNode, TEXT("ReturnValue"));
+			}
+		} // Handle (array-find array item) - find index of item
+		if (FormName.Equals(TEXT("array-find"), ESearchCase::IgnoreCase) && Expr->Num() >= 3)
+		{
+			UFunction* FindFunc = UKismetArrayLibrary::StaticClass()->FindFunctionByName(TEXT("Array_Find"));
+			if (FindFunc)
+			{
+				UK2Node_CallFunction* FindNode = NewObject<UK2Node_CallFunction>(Ctx.Graph);
+				FindNode->SetFromFunction(FindFunc);
+				FindNode->NodePosX = Ctx.CurrentX;
+				FindNode->NodePosY = Ctx.CurrentY;
+				Ctx.Graph->AddNode(FindNode, false, false);
+				FindNode->AllocateDefaultPins(); FString TempId = Ctx.GenerateTempId();
+				Ctx.TempIdToNode.Add(TempId, FindNode); return FindOutputPin(FindNode, TEXT("ReturnValue"));
+			}
+		} // Handle vector constructors: (vec x y z)
+		if (FormName.Equals(TEXT("vec"), ESearchCase::IgnoreCase) && Expr->Num() >= 4)
+		{
+			UFunction* Function = UKismetMathLibrary::StaticClass()->FindFunctionByName(TEXT("MakeVector"));
+			if (Function)
+			{
+				UK2Node_CallFunction* VecNode = NewObject<UK2Node_CallFunction>(Ctx.Graph);
+				VecNode->SetFromFunction(Function);
+				VecNode->NodePosX = Ctx.CurrentX;
+				VecNode->NodePosY = Ctx.CurrentY;
+				Ctx.Graph->AddNode(VecNode, false, false);
+				VecNode->AllocateDefaultPins();
+				EnsureNodeHasValidGuid(VecNode);
+				
+				// Set X, Y, Z
+				const TCHAR* PinNames[] = { TEXT("X"), TEXT("Y"), TEXT("Z") };
+				for (int32 i = 0; i < 3 && i + 1 < Expr->Num(); i++)
+				{
+					UEdGraphPin* Pin = VecNode->FindPin(PinNames[i]);
+					if (Pin)
+					{
+						FLispNodePtr ArgExpr = Expr->Get(i + 1);
+						
+						// First try to resolve as an expression (handles math ops, variable refs, etc.)
+						if (ArgExpr->IsList() || ArgExpr->IsSymbol())
+						{
+							UEdGraphPin* ArgSource = ResolveLispExpression(ArgExpr, Ctx);
+							if (ArgSource)
+							{
+								ConnectNodes(ArgSource, Pin, Ctx);
+								continue; // Successfully connected, move to next pin
+							}
+						}
+						
+						// Fall back to literal number
+						if (ArgExpr->IsNumber())
+						{
+							SetNumericPinDefaultValue(Pin, ArgExpr->NumberValue);
+						}
+						else if (ArgExpr->IsString())
+						{
+							Pin->DefaultValue = ArgExpr->StringValue;
+						}
+					}
+				}
+				
+				FString TempId = Ctx.GenerateTempId();
+				Ctx.TempIdToNode.Add(TempId, VecNode);
+				
+				return FindOutputPin(VecNode, TEXT("ReturnValue"));
+			}
+		}
+		
+		// Handle rotator constructors: (rot roll pitch yaw)
+		if (FormName.Equals(TEXT("rot"), ESearchCase::IgnoreCase) && Expr->Num() >= 4)
+		{
+			UFunction* Function = UKismetMathLibrary::StaticClass()->FindFunctionByName(TEXT("MakeRotator"));
+			if (Function)
+			{
+				UK2Node_CallFunction* RotNode = NewObject<UK2Node_CallFunction>(Ctx.Graph);
+				RotNode->SetFromFunction(Function);
+				RotNode->NodePosX = Ctx.CurrentX;
+				RotNode->NodePosY = Ctx.CurrentY;
+				Ctx.Graph->AddNode(RotNode, false, false);
+				RotNode->AllocateDefaultPins();
+				EnsureNodeHasValidGuid(RotNode);
+				
+				// MakeRotator pin order: Roll, Pitch, Yaw
+				const TCHAR* PinNames[] = { TEXT("Roll"), TEXT("Pitch"), TEXT("Yaw") };
+				for (int32 i = 0; i < 3 && i + 1 < Expr->Num(); i++)
+				{
+					UEdGraphPin* Pin = RotNode->FindPin(PinNames[i]);
+					if (Pin)
+					{
+						FLispNodePtr ArgExpr = Expr->Get(i + 1);
+						
+						// First try to resolve as an expression (handles math ops, variable refs, etc.)
+						if (ArgExpr->IsList() || ArgExpr->IsSymbol())
+						{
+							UEdGraphPin* ArgSource = ResolveLispExpression(ArgExpr, Ctx);
+							if (ArgSource)
+							{
+								ConnectNodes(ArgSource, Pin, Ctx);
+								continue; // Successfully connected, move to next pin
+							}
+						}
+						
+						// Fall back to literal number
+						if (ArgExpr->IsNumber())
+						{
+							SetNumericPinDefaultValue(Pin, ArgExpr->NumberValue);
+						}
+						else if (ArgExpr->IsString())
+						{
+							Pin->DefaultValue = ArgExpr->StringValue;
+						}
+					}
+				}
+				
+				FString TempId = Ctx.GenerateTempId();
+				Ctx.TempIdToNode.Add(TempId, RotNode);
+				
+				return FindOutputPin(RotNode, TEXT("ReturnValue"));
+			}
+		}
+		
+// Handle (call target func args...)
+		if (FormName.Equals(TEXT("call"), ESearchCase::IgnoreCase) && Expr->Num() >= 3)
+		{
+			// Create function call node
+			FString FuncName = Expr->Get(2)->IsSymbol() ? Expr->Get(2)->StringValue : TEXT(""); // Check if target is "self" - if so, also search Blueprint's own functions
+			bool bTargetIsSelf = false;
+			if (Expr->Get(1)->IsSymbol())
+			{
+				FString TargetName = Expr->Get(1)->StringValue;
+				bTargetIsSelf = TargetName.Equals(TEXT("self"), ESearchCase::IgnoreCase);
+			} UFunction* Function = nullptr;
+			
+			// First, check if this is a function defined in the Blueprint itself
+			if (bTargetIsSelf && Ctx.Blueprint)
+			{
+				// Look for the function in the Blueprint's function graphs
+				for (UEdGraph* FuncGraph : Ctx.Blueprint->FunctionGraphs)
+				{
+					if (FuncGraph && FuncGraph->GetFName() == FName(*FuncName))
+					{
+						// Found a matching function graph - create a CallFunctionNode for it
+						UK2Node_CallFunction* CallNode = NewObject<UK2Node_CallFunction>(Ctx.Graph);
+						CallNode->FunctionReference.SetSelfMember(FName(*FuncName));
+						CallNode->NodePosX = Ctx.CurrentX;
+						CallNode->NodePosY = Ctx.CurrentY;
+						Ctx.Graph->AddNode(CallNode, false, false);
+						CallNode->AllocateDefaultPins();
+						EnsureNodeHasValidGuid(CallNode); // Connect arguments
+						int32 ArgIndex = 0;
+						for (int32 i = 3; i < Expr->Num(); i++)
+						{
+							int32 PinIdx = 0;
+							for (UEdGraphPin* Pin : CallNode->Pins)
+							{
+								if (Pin->Direction == EGPD_Input && 
+									Pin->PinType.PinCategory != UEdGraphSchema_K2::PC_Exec &&
+									Pin->PinName != UEdGraphSchema_K2::PN_Self &&
+									!Pin->bHidden)
+								{
+									if (PinIdx == ArgIndex)
+									{
+										UEdGraphPin* ArgSource = ResolveLispExpression(Expr->Get(i), Ctx);
+										if (ArgSource)
+										{
+											ConnectNodes(ArgSource, Pin, Ctx);
+										}
+										else if (Expr->Get(i)->IsNumber())
+										{
+											SetNumericPinDefaultValue(Pin, Expr->Get(i)->NumberValue);
+										}
+										else if (Expr->Get(i)->IsString())
+										{
+											Pin->DefaultValue = Expr->Get(i)->StringValue;
+										}
+										ArgIndex++;
+										break;
+									}
+									PinIdx++;
+								}
+							}
+						} FString TempId = Ctx.GenerateTempId();
+						Ctx.TempIdToNode.Add(TempId, CallNode);
+						Ctx.AdvancePosition(); return FindOutputPin(CallNode, TEXT("ReturnValue"));
+					}
+				}
+			}
+			
+			// Search common classes for the function
+			// IMPORTANT: UKismetMathLibrary FIRST for math functions (Sin, Cos, FFloor, etc.)
+			TArray<UClass*> ClassesToSearch = {
+				UKismetMathLibrary::StaticClass(),
+				UKismetSystemLibrary::StaticClass(),
+				UKismetStringLibrary::StaticClass(),
+				UKismetArrayLibrary::StaticClass(),
+				UGameplayStatics::StaticClass(),
+				AActor::StaticClass(),
+				APawn::StaticClass(),
+				ACharacter::StaticClass(),
+				USceneComponent::StaticClass(),
+				UPrimitiveComponent::StaticClass(),
+				Ctx.Blueprint->ParentClass,
+				Ctx.Blueprint->GeneratedClass
+			};
+			
+			// Add Geometry Script classes for procedural mesh functions
+			AddGeometryScriptClasses(ClassesToSearch);
+			
+			for (UClass* Class : ClassesToSearch)
+			{
+				if (Class)
+				{
+					Function = Class->FindFunctionByName(*FuncName);
+					if (Function) break;
+					// Also try with K2_ prefix
+					Function = Class->FindFunctionByName(*FString::Printf(TEXT("K2_%s"), *FuncName));
+					if (Function) break;
+				}
+			}
+			
+			if (Function)
+			{
+				UK2Node_CallFunction* CallNode = NewObject<UK2Node_CallFunction>(Ctx.Graph);
+				CallNode->SetFromFunction(Function);
+				CallNode->NodePosX = Ctx.CurrentX;
+				CallNode->NodePosY = Ctx.CurrentY;
+				Ctx.Graph->AddNode(CallNode, false, false);
+				CallNode->AllocateDefaultPins();
+				EnsureNodeHasValidGuid(CallNode);
+				
+				// Connect target
+				UEdGraphPin* TargetPin = CallNode->FindPin(UEdGraphSchema_K2::PN_Self);
+				if (TargetPin)
+				{
+					UEdGraphPin* TargetSource = ResolveLispExpression(Expr->Get(1), Ctx);
+					if (TargetSource)
+					{
+						ConnectNodes(TargetSource, TargetPin, Ctx);
+					}
+				}
+				
+				// Connect arguments
+				int32 ArgIndex = 0;
+				for (int32 i = 3; i < Expr->Num(); i++)
+				{
+					// Find the next input pin that's not exec or self
+					for (UEdGraphPin* Pin : CallNode->Pins)
+					{
+						if (Pin->Direction == EGPD_Input && 
+							Pin->PinType.PinCategory != UEdGraphSchema_K2::PC_Exec &&
+							Pin->PinName != UEdGraphSchema_K2::PN_Self &&
+							!Pin->bHidden)
+						{
+							if (ArgIndex == i - 3)
+							{
+								UEdGraphPin* ArgSource = ResolveLispExpression(Expr->Get(i), Ctx);
+								if (ArgSource)
+								{
+									ConnectNodes(ArgSource, Pin, Ctx);
+								}
+								else if (Expr->Get(i)->IsNumber())
+								{
+									double NumVal = Expr->Get(i)->NumberValue;
+									// Check if the pin expects an integer
+									if (Pin->PinType.PinCategory == UEdGraphSchema_K2::PC_Int)
+									{
+										Pin->DefaultValue = FString::Printf(TEXT("%d"), FMath::RoundToInt(NumVal));
+									}
+									else
+									{
+										SetNumericPinDefaultValue(Pin, NumVal);
+									}
+								}
+								else if (Expr->Get(i)->IsString())
+								{
+									Pin->DefaultValue = Expr->Get(i)->StringValue;
+								}
+								break;
+							}
+							ArgIndex++;
+						}
+					}
+				}
+				
+				FString TempId = Ctx.GenerateTempId();
+				Ctx.TempIdToNode.Add(TempId, CallNode);
+				Ctx.AdvancePosition();
+				
+				return FindOutputPin(CallNode, TEXT("ReturnValue"));
+			}
+		}
+		
+		// Handle direct function call (FuncName args...)
+		if (!FormName.IsEmpty())
+		{
+			// Map of shorthand function names to UE function names
+			static TMap<FString, FString> PureFuncShorthand;
+			if (PureFuncShorthand.Num() == 0)
+			{
+				// Math - Note: UE uses capitalized names
+				PureFuncShorthand.Add(TEXT("sin"), TEXT("Sin"));
+				PureFuncShorthand.Add(TEXT("cos"), TEXT("Cos"));
+				PureFuncShorthand.Add(TEXT("tan"), TEXT("Tan"));
+				PureFuncShorthand.Add(TEXT("asin"), TEXT("Asin"));
+				PureFuncShorthand.Add(TEXT("acos"), TEXT("Acos"));
+				PureFuncShorthand.Add(TEXT("atan"), TEXT("Atan"));
+				PureFuncShorthand.Add(TEXT("atan2"), TEXT("Atan2"));
+				PureFuncShorthand.Add(TEXT("abs"), TEXT("Abs"));
+				PureFuncShorthand.Add(TEXT("sqrt"), TEXT("Sqrt"));
+				PureFuncShorthand.Add(TEXT("sqr"), TEXT("Square"));
+				PureFuncShorthand.Add(TEXT("square"), TEXT("Square"));
+				PureFuncShorthand.Add(TEXT("pow"), TEXT("Pow"));
+				PureFuncShorthand.Add(TEXT("exp"), TEXT("Exp"));
+				PureFuncShorthand.Add(TEXT("log"), TEXT("Loge"));
+				PureFuncShorthand.Add(TEXT("log10"), TEXT("Log10"));
+				PureFuncShorthand.Add(TEXT("floor"), TEXT("FFloor"));
+				PureFuncShorthand.Add(TEXT("Floor"), TEXT("FFloor"));
+				PureFuncShorthand.Add(TEXT("ceil"), TEXT("FCeil"));
+				PureFuncShorthand.Add(TEXT("Ceil"), TEXT("FCeil"));
+				PureFuncShorthand.Add(TEXT("round"), TEXT("Round"));
+				PureFuncShorthand.Add(TEXT("Round"), TEXT("Round"));
+				PureFuncShorthand.Add(TEXT("trunc"), TEXT("FTrunc"));
+				PureFuncShorthand.Add(TEXT("Trunc"), TEXT("FTrunc"));
+				PureFuncShorthand.Add(TEXT("frac"), TEXT("Fraction"));
+				PureFuncShorthand.Add(TEXT("Frac"), TEXT("Fraction"));
+				PureFuncShorthand.Add(TEXT("fmod"), TEXT("FMod"));
+				PureFuncShorthand.Add(TEXT("FMod"), TEXT("FMod"));
+				PureFuncShorthand.Add(TEXT("sign"), TEXT("SignOfFloat"));
+				PureFuncShorthand.Add(TEXT("min"), TEXT("FMin"));
+				PureFuncShorthand.Add(TEXT("max"), TEXT("FMax"));
+				PureFuncShorthand.Add(TEXT("clamp"), TEXT("FClamp"));
+				PureFuncShorthand.Add(TEXT("lerp"), TEXT("Lerp"));
+				PureFuncShorthand.Add(TEXT("random"), TEXT("RandomFloat"));
+				PureFuncShorthand.Add(TEXT("rand"), TEXT("RandomFloat"));
+			}
+			
+			// Check if this is a shorthand
+			FString ActualFuncName = FormName;
+			if (FString* MappedName = PureFuncShorthand.Find(FormName))
+			{
+				ActualFuncName = *MappedName;
+			}
+			
+			UFunction* Function = nullptr;
+			TArray<UClass*> ClassesToSearch = {
+				UKismetMathLibrary::StaticClass(),
+				UKismetSystemLibrary::StaticClass(),
+				UKismetStringLibrary::StaticClass(),
+				UKismetArrayLibrary::StaticClass(),
+				UGameplayStatics::StaticClass(),
+				AActor::StaticClass(),
+				APawn::StaticClass(),
+				ACharacter::StaticClass(),
+				Ctx.Blueprint->ParentClass
+			};
+			
+			// Add Geometry Script classes for procedural mesh functions
+			AddGeometryScriptClasses(ClassesToSearch);
+			
+			for (UClass* Class : ClassesToSearch)
+			{
+				if (Class)
+				{
+					Function = Class->FindFunctionByName(*ActualFuncName);
+					if (Function) break;
+					// Also try original name if shorthand didn't work
+					if (ActualFuncName != FormName)
+					{
+						Function = Class->FindFunctionByName(*FormName);
+						if (Function) break;
+					}
+				}
+			}
+			
+			if (Function)
+			{
+				UK2Node_CallFunction* CallNode = NewObject<UK2Node_CallFunction>(Ctx.Graph);
+				CallNode->SetFromFunction(Function);
+				CallNode->NodePosX = Ctx.CurrentX;
+				CallNode->NodePosY = Ctx.CurrentY;
+				Ctx.Graph->AddNode(CallNode, false, false);
+				CallNode->AllocateDefaultPins();
+				EnsureNodeHasValidGuid(CallNode);
+				
+				// Connect arguments
+				int32 ArgIndex = 0;
+				for (int32 i = 1; i < Expr->Num(); i++)
+				{
+					int32 PinArgIndex = 0;
+					for (UEdGraphPin* Pin : CallNode->Pins)
+					{
+						if (Pin->Direction == EGPD_Input && 
+							Pin->PinType.PinCategory != UEdGraphSchema_K2::PC_Exec &&
+							Pin->PinName != UEdGraphSchema_K2::PN_Self &&
+							!Pin->bHidden)
+						{
+							if (PinArgIndex == ArgIndex)
+							{
+								UEdGraphPin* ArgSource = ResolveLispExpression(Expr->Get(i), Ctx);
+								if (ArgSource)
+								{
+									ConnectNodes(ArgSource, Pin, Ctx);
+								}
+								else if (Expr->Get(i)->IsNumber())
+								{
+									SetNumericPinDefaultValue(Pin, Expr->Get(i)->NumberValue);
+								}
+								else if (Expr->Get(i)->IsString())
+								{
+									Pin->DefaultValue = Expr->Get(i)->StringValue;
+								}
+								else if (Expr->Get(i)->IsSymbol())
+								{
+									// Could be true/false
+									if (Expr->Get(i)->StringValue.Equals(TEXT("true"), ESearchCase::IgnoreCase))
+									{
+										Pin->DefaultValue = TEXT("true");
+									}
+									else if (Expr->Get(i)->StringValue.Equals(TEXT("false"), ESearchCase::IgnoreCase))
+									{
+										Pin->DefaultValue = TEXT("false");
+									}
+								}
+								ArgIndex++;
+								break;
+							}
+							PinArgIndex++;
+						}
+					}
+				}
+				
+				FString TempId = Ctx.GenerateTempId();
+				Ctx.TempIdToNode.Add(TempId, CallNode);
+				Ctx.AdvancePosition();
+				
+				return FindOutputPin(CallNode, TEXT("ReturnValue"));
+			}
+		}
+	}
+	
+	Ctx.Warnings.Add(FString::Printf(TEXT("Could not resolve expression: %s"), *Expr->ToString(false)));
+	return nullptr;
+}
+
+// Convert execution body (seq or single statement)
+static void ConvertExecBody(const FLispNodePtr& Body, FLispToBPContext& Ctx, UEdGraphPin*& CurrentExecPin)
+{
+	if (!Body.IsValid() || Body->IsNil())
+	{
+		return;
+	}
+	
+	// Handle (seq statement1 statement2 ...)
+	if (Body->IsList() && Body->GetFormName().Equals(TEXT("seq"), ESearchCase::IgnoreCase))
+	{
+		for (int32 i = 1; i < Body->Num(); i++)
+		{
+			UEdGraphPin* NextExecPin = nullptr;
+			UEdGraphNode* Node = ConvertLispFormToNode(Body->Get(i), Ctx, NextExecPin);
+			
+			if (Node && CurrentExecPin)
+			{
+				UEdGraphPin* NodeInput = GetNodeExecInput(Node);
+				if (NodeInput)
+				{
+					ConnectNodes(CurrentExecPin, NodeInput, Ctx);
+				}
+			}
+			
+			if (NextExecPin)
+			{
+				CurrentExecPin = NextExecPin;
+			}
+			else if (Node)
+			{
+				// Don't try to get exec output from branch nodes - they have no single continuation
+				UK2Node_IfThenElse* BranchCheck = Cast<UK2Node_IfThenElse>(Node);
+				if (!BranchCheck)
+				{
+					CurrentExecPin = GetNodeExecOutput(Node);
+				}
+				// For branch nodes, CurrentExecPin stays null - the branch bodies handle their own flow
+			}
+		}
+	}
+	else
+	{
+		// Single statement
+		Ctx.Warnings.Add(FString::Printf(TEXT("DEBUG ConvertExecBody: Processing single statement: %s"), *Body->ToString(false).Left(100)));
+		
+		UEdGraphPin* NextExecPin = nullptr;
+		UEdGraphNode* Node = ConvertLispFormToNode(Body, Ctx, NextExecPin);
+		
+		if (Node)
+		{
+			Ctx.Warnings.Add(FString::Printf(TEXT("DEBUG ConvertExecBody: Created node %s"), *Node->GetClass()->GetName()));
+			
+			if (CurrentExecPin)
+			{
+				UEdGraphPin* NodeInput = GetNodeExecInput(Node);
+				if (NodeInput)
+				{
+					ConnectNodes(CurrentExecPin, NodeInput, Ctx);
+					Ctx.Warnings.Add(TEXT("DEBUG ConvertExecBody: Connected exec pin to node input"));
+				}
+				else
+				{
+					Ctx.Warnings.Add(TEXT("DEBUG ConvertExecBody: Node has no exec input pin!"));
+				}
+			}
+			else
+			{
+				Ctx.Warnings.Add(TEXT("DEBUG ConvertExecBody: CurrentExecPin is null!"));
+			}
+		}
+		else
+		{
+			Ctx.Warnings.Add(FString::Printf(TEXT("DEBUG ConvertExecBody: ConvertLispFormToNode returned NULL for: %s"), *Body->ToString(false).Left(100)));
+		}
+		
+		if (NextExecPin)
+		{
+			CurrentExecPin = NextExecPin;
+		}
+		else if (Node)
+		{
+			// Don't try to get exec output from branch nodes - they have no single continuation
+			UK2Node_IfThenElse* BranchCheck = Cast<UK2Node_IfThenElse>(Node);
+			if (!BranchCheck)
+			{
+				CurrentExecPin = GetNodeExecOutput(Node);
+			}
+		}
+	}
+}
+
+// Convert a Lisp form to a Blueprint node
+static UEdGraphNode* ConvertLispFormToNode(const FLispNodePtr& Form, FLispToBPContext& Ctx, UEdGraphPin*& OutExecPin)
+{
+	OutExecPin = nullptr;
+	
+	if (!Form.IsValid() || !Form->IsList() || Form->Num() == 0)
+	{
+		return nullptr;
+	}
+	
+	FString FormName = Form->GetFormName();
+	
+	// Handle (seq statement1 statement2 ...) - execute statements in sequence
+	if (FormName.Equals(TEXT("seq"), ESearchCase::IgnoreCase))
+	{
+		// Process each statement in the sequence
+		UEdGraphPin* CurrentExecPin = nullptr;
+		UEdGraphNode* FirstNode = nullptr;
+		
+		for (int32 i = 1; i < Form->Num(); i++)
+		{
+			FLispNodePtr Statement = Form->Get(i);
+			if (!Statement.IsValid() || Statement->IsNil()) continue;
+			
+			UEdGraphPin* StatementExecOut = nullptr;
+			UEdGraphNode* StatementNode = ConvertLispFormToNode(Statement, Ctx, StatementExecOut);
+			
+			if (StatementNode)
+			{
+				if (!FirstNode)
+				{
+					FirstNode = StatementNode;
+				}
+				
+				// Connect previous statement's output to this statement's input
+				if (CurrentExecPin)
+				{
+					UEdGraphPin* StatementExecIn = GetNodeExecInput(StatementNode);
+					if (StatementExecIn)
+					{
+						ConnectNodes(CurrentExecPin, StatementExecIn, Ctx);
+					}
+				}
+				
+				// Update current exec pin to this statement's output
+				if (StatementExecOut)
+				{
+					CurrentExecPin = StatementExecOut;
+				}
+				else if (StatementNode)
+				{
+					// Don't try to get exec output from branch nodes - they have no single continuation
+					UK2Node_IfThenElse* BranchCheck = Cast<UK2Node_IfThenElse>(StatementNode);
+					if (!BranchCheck)
+					{
+						CurrentExecPin = GetNodeExecOutput(StatementNode);
+					}
+					else
+					{
+						// For branch nodes, set CurrentExecPin to nullptr since branches have no single continuation
+						CurrentExecPin = nullptr;
+					}
+				}
+			}
+		}
+		
+		// Return the first node and the last exec output
+		OutExecPin = CurrentExecPin;
+		return FirstNode;
+	}
+	
+	// Handle (branch condition :true body :false body)
+	if (FormName.Equals(TEXT("branch"), ESearchCase::IgnoreCase))
+	{
+		UK2Node_IfThenElse* BranchNode = NewObject<UK2Node_IfThenElse>(Ctx.Graph);
+		BranchNode->NodePosX = Ctx.CurrentX;
+		BranchNode->NodePosY = Ctx.CurrentY;
+		Ctx.Graph->AddNode(BranchNode, false, false);
+		BranchNode->AllocateDefaultPins();
+		Ctx.AdvancePosition();
+		
+		// Connect condition
+		if (Form->Num() > 1)
+		{
+			UEdGraphPin* ConditionPin = BranchNode->GetConditionPin();
+			UEdGraphPin* ConditionSource = ResolveLispExpression(Form->Get(1), Ctx);
+			if (ConditionSource && ConditionPin)
+			{
+				ConnectNodes(ConditionSource, ConditionPin, Ctx);
+			}
+		}
+		
+		// Handle true branch
+		FLispNodePtr TrueBody = Form->GetKeywordArg(TEXT(":true"));
+		if (TrueBody.IsValid() && !TrueBody->IsNil())
+		{
+			UEdGraphPin* TrueExecPin = BranchNode->GetThenPin();
+			Ctx.Warnings.Add(FString::Printf(TEXT("DEBUG branch: Processing :true body: %s"), *TrueBody->ToString(false).Left(100)));
+			Ctx.Warnings.Add(FString::Printf(TEXT("DEBUG branch: TrueExecPin (Then) = %s, linked_to=%d"), 
+				TrueExecPin ? *TrueExecPin->PinName.ToString() : TEXT("NULL"),
+				TrueExecPin ? TrueExecPin->LinkedTo.Num() : -1));
+			
+			ConvertExecBody(TrueBody, Ctx, TrueExecPin);
+			
+			Ctx.Warnings.Add(FString::Printf(TEXT("DEBUG branch: After ConvertExecBody, TrueExecPin linked_to=%d"), 
+				TrueExecPin ? TrueExecPin->LinkedTo.Num() : -1));
+			if (TrueExecPin && TrueExecPin->LinkedTo.Num() > 0)
+			{
+				UEdGraphPin* LinkedPin = TrueExecPin->LinkedTo[0];
+				Ctx.Warnings.Add(FString::Printf(TEXT("DEBUG branch: Then linked to pin '%s' on node '%s' (GUID=%s)"), 
+					*LinkedPin->PinName.ToString(),
+					*LinkedPin->GetOwningNode()->GetClass()->GetName(),
+					*LinkedPin->GetOwningNode()->NodeGuid.ToString()));
+			}
+		}
+		else
+		{
+			Ctx.Warnings.Add(TEXT("DEBUG branch: :true body is nil or invalid"));
+		}
+		
+		// Handle false branch
+		FLispNodePtr FalseBody = Form->GetKeywordArg(TEXT(":false"));
+		if (FalseBody.IsValid() && !FalseBody->IsNil())
+		{
+			UEdGraphPin* FalseExecPin = BranchNode->GetElsePin();
+			ConvertExecBody(FalseBody, Ctx, FalseExecPin);
+		}
+		
+		// IMPORTANT: Branch nodes don't have a single "continue" exec output.
+		// The Then and Else pins are already connected to their respective bodies.
+		// Set OutExecPin to nullptr so the caller doesn't try to connect anything
+		// after the branch to the Then pin (which would break the :true body connection).
+		OutExecPin = nullptr;
+		
+		Ctx.Warnings.Add(TEXT("DEBUG branch: Setting OutExecPin to nullptr (branch has no single continuation)"));
+		
+		return BranchNode;
+	}
+	
+	// Handle (set variable value)
+	if (FormName.Equals(TEXT("set"), ESearchCase::IgnoreCase) && Form->Num() >= 3)
+	{
+		FString VarName = Form->Get(1)->IsSymbol() ? Form->Get(1)->StringValue : TEXT("");
+		
+		UK2Node_VariableSet* SetNode = NewObject<UK2Node_VariableSet>(Ctx.Graph);
+		SetNode->VariableReference.SetSelfMember(FName(*VarName));
+		SetNode->NodePosX = Ctx.CurrentX;
+		SetNode->NodePosY = Ctx.CurrentY;
+		Ctx.Graph->AddNode(SetNode, false, false);
+		SetNode->AllocateDefaultPins();
+		Ctx.AdvancePosition();
+		
+		// Connect value
+		for (UEdGraphPin* Pin : SetNode->Pins)
+		{
+			if (Pin->Direction == EGPD_Input && 
+				Pin->PinType.PinCategory != UEdGraphSchema_K2::PC_Exec &&
+				Pin->PinName != UEdGraphSchema_K2::PN_Self)
+			{
+				UEdGraphPin* ValueSource = ResolveLispExpression(Form->Get(2), Ctx);
+				if (ValueSource)
+				{
+					// Check if we need type conversion (e.g., float to int)
+					if (Pin->PinType.PinCategory == UEdGraphSchema_K2::PC_Int &&
+						(ValueSource->PinType.PinCategory == UEdGraphSchema_K2::PC_Real ||
+						 ValueSource->PinType.PinCategory == UEdGraphSchema_K2::PC_Float ||
+						 ValueSource->PinType.PinCategory == UEdGraphSchema_K2::PC_Double))
+					{
+						// Create float to int conversion using FTrunc (truncate to integer)
+						UFunction* ConvFunc = UKismetMathLibrary::StaticClass()->FindFunctionByName(TEXT("FTrunc"));
+						if (!ConvFunc)
+						{
+							ConvFunc = UKismetMathLibrary::StaticClass()->FindFunctionByName(TEXT("FFloor"));
+						}
+						if (ConvFunc)
+						{
+							UK2Node_CallFunction* ConvNode = NewObject<UK2Node_CallFunction>(Ctx.Graph);
+							ConvNode->SetFromFunction(ConvFunc);
+							ConvNode->NodePosX = Ctx.CurrentX - 150;
+							ConvNode->NodePosY = Ctx.CurrentY;
+							Ctx.Graph->AddNode(ConvNode, false, false);
+							ConvNode->AllocateDefaultPins();
+							EnsureNodeHasValidGuid(ConvNode);
+							
+							// FTrunc/FFloor have an "A" input pin for the float value
+							UEdGraphPin* InputPin = ConvNode->FindPin(TEXT("A"));
+							if (!InputPin)
+							{
+								// Try finding first non-exec input
+								for (UEdGraphPin* ConvPin : ConvNode->Pins)
+								{
+									if (ConvPin->Direction == EGPD_Input && 
+										ConvPin->PinType.PinCategory != UEdGraphSchema_K2::PC_Exec &&
+										ConvPin->PinName != UEdGraphSchema_K2::PN_Self)
+									{
+										InputPin = ConvPin;
+										break;
+									}
+								}
+							}
+							
+							if (InputPin)
+							{
+								ConnectNodes(ValueSource, InputPin, Ctx);
+							}
+							
+							// Connect conversion output to target
+							UEdGraphPin* ConvOutput = FindOutputPin(ConvNode, TEXT("ReturnValue"));
+							if (ConvOutput)
+							{
+								ConnectNodes(ConvOutput, Pin, Ctx);
+							}
+						}
+						else
+						{
+							ConnectNodes(ValueSource, Pin, Ctx);
+						}
+					}
+					else
+					{
+						ConnectNodes(ValueSource, Pin, Ctx);
+					}
+				}
+				else if (Form->Get(2)->IsNumber())
+				{
+					double NumVal = Form->Get(2)->NumberValue;
+					// Check if the pin expects an integer
+					if (Pin->PinType.PinCategory == UEdGraphSchema_K2::PC_Int)
+					{
+						Pin->DefaultValue = FString::Printf(TEXT("%d"), FMath::RoundToInt(NumVal));
+					}
+					else
+					{
+						SetNumericPinDefaultValue(Pin, NumVal);
+					}
+				}
+				else if (Form->Get(2)->IsString())
+				{
+					Pin->DefaultValue = Form->Get(2)->StringValue;
+				}
+				else if (Form->Get(2)->IsSymbol())
+				{
+					// Handle boolean literals: true, false
+					FString SymValue = Form->Get(2)->StringValue;
+					if (SymValue.Equals(TEXT("true"), ESearchCase::IgnoreCase))
+					{
+						Pin->DefaultValue = TEXT("true");
+					}
+					else if (SymValue.Equals(TEXT("false"), ESearchCase::IgnoreCase))
+					{
+						Pin->DefaultValue = TEXT("false");
+					}
+					else if (SymValue.Equals(TEXT("nil"), ESearchCase::IgnoreCase))
+					{
+						// nil = None/null
+						Pin->DefaultValue = TEXT("");
+					}
+					else
+					{
+						// Could be an enum value or other symbol - try setting it directly
+						Pin->DefaultValue = SymValue;
+					}
+				}
+				break;
+			}
+		}
+		
+		OutExecPin = GetNodeExecOutput(SetNode);
+		return SetNode;
+	}
+	
+	// Handle (let ...) - supports both simple and Common Lisp-style syntax
+	// Simple: (let var expr)
+	// Common Lisp: (let ((var1 expr1) (var2 expr2)) body...)
+	if (FormName.Equals(TEXT("let"), ESearchCase::IgnoreCase) && Form->Num() >= 2)
+	{
+		// Check for Common Lisp-style: (let ((var expr) ...) body...)
+		// The second element would be a list of binding lists
+		if (Form->Get(1)->IsList())
+		{
+			FLispNodePtr BindingsList = Form->Get(1);
+			
+			// Process each binding: ((var1 expr1) (var2 expr2) ...)
+			for (int32 i = 0; i < BindingsList->Num(); i++)
+			{
+				FLispNodePtr Binding = BindingsList->Get(i);
+				if (Binding->IsList() && Binding->Num() >= 2)
+				{
+					FString VarName = Binding->Get(0)->IsSymbol() ? Binding->Get(0)->StringValue : TEXT("");
+					
+					if (!VarName.IsEmpty())
+					{
+						UEdGraphPin* ExprResult = ResolveLispExpression(Binding->Get(1), Ctx);
+						if (ExprResult)
+						{
+							UEdGraphNode* SourceNode = ExprResult->GetOwningNode();
+							EnsureNodeHasValidGuid(SourceNode);
+							FString NodeGuidStr = SourceNode->NodeGuid.ToString();
+							FString PinNameStr = ExprResult->PinName.ToString();
+							
+							if (!Ctx.VariableToNodeId.Contains(VarName))
+							{
+								Ctx.VariableToNodeId.Add(VarName, NodeGuidStr);
+								Ctx.VariableToPin.Add(VarName, PinNameStr);
+								if (!Ctx.TempIdToNode.Contains(NodeGuidStr))
+								{
+									Ctx.TempIdToNode.Add(NodeGuidStr, SourceNode);
+								}
+								// Also add with variable name prefixed key for direct lookup
+								FString VarKey = TEXT("_var_") + VarName;
+								Ctx.TempIdToNode.Add(VarKey, SourceNode);
+							}
+						}
+						else
+						{
+							Ctx.Warnings.Add(FString::Printf(TEXT("let '%s': Could not resolve expression"), *VarName));
+						}
+					}
+				}
+			}
+			
+			// Process body statements (everything after the bindings list)
+			for (int32 i = 2; i < Form->Num(); i++)
+			{
+				UEdGraphPin* BodyExecOut = nullptr;
+				ConvertLispFormToNode(Form->Get(i), Ctx, BodyExecOut);
+			}
+			
+			return nullptr;
+		}
+		
+		// Simple style: (let var expr)
+		if (Form->Num() < 3)
+		{
+			Ctx.Warnings.Add(TEXT("let requires variable and expression: (let var expr)"));
+			return nullptr;
+		}
+		
+		FString VarName = Form->Get(1)->IsSymbol() ? Form->Get(1)->StringValue : TEXT("");
+		FLispNodePtr ExprNode = Form->Get(2);
+		
+		// Handle literal values by creating a constant node
+		if (ExprNode->IsNumber() || ExprNode->IsString())
+		{
+			// For literals, we need to store them differently - create a literal node
+			// that can be used as a source when the variable is referenced
+			if (ExprNode->IsNumber())
+			{
+				// Create a MakeLiteralFloat node for numeric literals
+				UFunction* LiteralFunc = UKismetSystemLibrary::StaticClass()->FindFunctionByName(TEXT("MakeLiteralFloat"));
+				if (!LiteralFunc)
+				{
+					LiteralFunc = UKismetMathLibrary::StaticClass()->FindFunctionByName(TEXT("Conv_FloatToDouble"));
+				}
+				
+				// Just store the literal value directly - we'll handle it in symbol resolution
+				// by storing special metadata
+				FString LiteralKey = TEXT("_literal_") + VarName;
+				Ctx.VariableToNodeId.Add(VarName, LiteralKey);
+				Ctx.VariableToPin.Add(VarName, FString::SanitizeFloat(ExprNode->NumberValue));
+				// Mark as literal by using special prefix
+				return nullptr;
+			}
+			else if (ExprNode->IsString())
+			{
+				FString LiteralKey = TEXT("_literalstr_") + VarName;
+				Ctx.VariableToNodeId.Add(VarName, LiteralKey);
+				Ctx.VariableToPin.Add(VarName, ExprNode->StringValue);
+				return nullptr;
+			}
+		}
+		
+		// Resolve the expression and store the mapping
+		UEdGraphPin* ExprResult = ResolveLispExpression(ExprNode, Ctx);
+		if (ExprResult)
+		{
+			UEdGraphNode* SourceNode = ExprResult->GetOwningNode();
+			EnsureNodeHasValidGuid(SourceNode);
+			FString NodeGuidStr = SourceNode->NodeGuid.ToString();
+			FString PinNameStr = ExprResult->PinName.ToString();
+			
+			// Store the variable mapping
+			// Use variable name as an additional key in TempIdToNode for direct lookup
+			if (!Ctx.VariableToNodeId.Contains(VarName))
+			{
+				Ctx.VariableToNodeId.Add(VarName, NodeGuidStr);
+				Ctx.VariableToPin.Add(VarName, PinNameStr);
+				
+				// Add to TempIdToNode using BOTH the NodeGuid AND the variable name as keys
+				// This ensures we can find the node by either method
+				if (!Ctx.TempIdToNode.Contains(NodeGuidStr))
+				{
+					Ctx.TempIdToNode.Add(NodeGuidStr, SourceNode);
+				}
+				// Also add with variable name prefixed key for direct lookup
+				FString VarKey = TEXT("_var_") + VarName;
+				Ctx.TempIdToNode.Add(VarKey, SourceNode);
+			}
+			else
+			{
+				Ctx.Warnings.Add(FString::Printf(TEXT("Variable '%s' already defined, ignoring redefinition"), *VarName));
+			}
+		}
+		else
+		{
+			Ctx.Warnings.Add(FString::Printf(TEXT("let '%s': Could not resolve expression: %s"), *VarName, *ExprNode->ToString(false)));
+		}
+		
+		// let doesn't create an exec node itself
+		return nullptr;
+	}
+	
+	// Handle (call target func args...) for impure functions
+	if (FormName.Equals(TEXT("call"), ESearchCase::IgnoreCase) && Form->Num() >= 3)
+	{
+		FString FuncName = Form->Get(2)->IsSymbol() ? Form->Get(2)->StringValue : TEXT(""); // Check if target is "self" - if so, also search Blueprint's own functions
+		bool bTargetIsSelf = false;
+		if (Form->Get(1)->IsSymbol())
+		{
+			FString TargetName = Form->Get(1)->StringValue;
+			bTargetIsSelf = TargetName.Equals(TEXT("self"), ESearchCase::IgnoreCase);
+		} // First, check if this is a function defined in the Blueprint itself
+		if (bTargetIsSelf && Ctx.Blueprint)
+		{
+			for (UEdGraph* FuncGraph : Ctx.Blueprint->FunctionGraphs)
+			{
+				if (FuncGraph && FuncGraph->GetFName() == FName(*FuncName))
+				{
+					// Found a matching function graph - create a CallFunctionNode for it
+					UK2Node_CallFunction* CallNode = NewObject<UK2Node_CallFunction>(Ctx.Graph);
+					CallNode->FunctionReference.SetSelfMember(FName(*FuncName));
+					CallNode->NodePosX = Ctx.CurrentX;
+					CallNode->NodePosY = Ctx.CurrentY;
+					Ctx.Graph->AddNode(CallNode, false, false);
+					CallNode->AllocateDefaultPins();
+					EnsureNodeHasValidGuid(CallNode);
+					Ctx.AdvancePosition(); // Connect arguments using SetPinValueFromExpr
+					int32 ArgIndex = 0;
+					for (int32 i = 3; i < Form->Num(); i++)
+					{
+						if (Form->Get(i)->IsKeyword()) continue; int32 PinIdx = 0;
+						for (UEdGraphPin* Pin : CallNode->Pins)
+						{
+							if (Pin->Direction == EGPD_Input && 
+								Pin->PinType.PinCategory != UEdGraphSchema_K2::PC_Exec &&
+								Pin->PinName != UEdGraphSchema_K2::PN_Self &&
+								!Pin->bHidden)
+							{
+								if (PinIdx == ArgIndex)
+								{
+									SetPinValueFromExpr(Pin, Form->Get(i), Ctx);
+									ArgIndex++;
+									break;
+								}
+								PinIdx++;
+							}
+						}
+					} FString TempId = Ctx.GenerateTempId();
+					Ctx.TempIdToNode.Add(TempId, CallNode); OutExecPin = GetNodeExecOutput(CallNode);
+					return CallNode;
+				}
+			}
+		} // Build list of function name variations to try
+		TArray<FString> FuncNamesToTry;
+		FuncNamesToTry.Add(FuncName);
+		FuncNamesToTry.Add(TEXT("K2_") + FuncName);  // Try with K2_ prefix
+		if (FuncName.StartsWith(TEXT("K2_")))
+		{
+			FuncNamesToTry.Add(FuncName.Mid(3));  // Try without K2_ prefix
+		}
+		
+		UFunction* Function = nullptr;
+		TArray<UClass*> ClassesToSearch = {
+			UKismetSystemLibrary::StaticClass(),
+			UGameplayStatics::StaticClass(),
+			AActor::StaticClass(),
+			APawn::StaticClass(),
+			ACharacter::StaticClass(),
+			UPrimitiveComponent::StaticClass(),
+			UStaticMeshComponent::StaticClass(),
+			USkeletalMeshComponent::StaticClass(),
+			USceneComponent::StaticClass(),
+			UActorComponent::StaticClass(),
+			Ctx.Blueprint->ParentClass,
+			Ctx.Blueprint->GeneratedClass
+		};
+		
+		// Add Geometry Script classes for procedural mesh functions
+		AddGeometryScriptClasses(ClassesToSearch);
+		
+		for (const FString& NameToTry : FuncNamesToTry)
+		{
+			for (UClass* Class : ClassesToSearch)
+			{
+				if (Class)
+				{
+					Function = Class->FindFunctionByName(*NameToTry);
+					if (Function) break;
+				}
+			}
+			if (Function) break;
+		}
+		
+		// If still not found, try to find the function by iterating all UObject classes
+		// This handles cases like component-specific methods
+		if (!Function)
+		{
+			// Try common component base classes
+			static TArray<UClass*> ComponentClasses;
+			if (ComponentClasses.Num() == 0)
+			{
+				ComponentClasses.Add(UMeshComponent::StaticClass());
+				ComponentClasses.Add(ULightComponent::StaticClass());
+				ComponentClasses.Add(UAudioComponent::StaticClass());
+				ComponentClasses.Add(UShapeComponent::StaticClass());
+			}
+			
+			for (UClass* Class : ComponentClasses)
+			{
+				if (Class)
+				{
+					Function = Class->FindFunctionByName(*FuncName);
+					if (Function) break;
+				}
+			}
+		}
+		
+		if (!Function)
+		{
+			Ctx.Warnings.Add(FString::Printf(TEXT("Function not found: %s (searched common classes)"), *FuncName));
+		}
+		
+		if (Function)
+		{
+			UK2Node_CallFunction* CallNode = NewObject<UK2Node_CallFunction>(Ctx.Graph);
+			CallNode->SetFromFunction(Function);
+			CallNode->NodePosX = Ctx.CurrentX;
+			CallNode->NodePosY = Ctx.CurrentY;
+			Ctx.Graph->AddNode(CallNode, false, false);
+			CallNode->AllocateDefaultPins();
+			EnsureNodeHasValidGuid(CallNode);  // IMPORTANT: Ensure valid GUID for exec connections
+			Ctx.AdvancePosition();
+			
+			FString TempId = Ctx.GenerateTempId();
+			Ctx.TempIdToNode.Add(TempId, CallNode);
+			
+			Ctx.Warnings.Add(FString::Printf(TEXT("DEBUG call: Created CallNode for %s, GUID=%s, has exec input=%s"), 
+				*FuncName, 
+				*CallNode->NodeGuid.ToString(),
+				GetNodeExecInput(CallNode) ? TEXT("YES") : TEXT("NO")));
+			
+			// Connect target
+			UEdGraphPin* TargetPin = CallNode->FindPin(UEdGraphSchema_K2::PN_Self);
+			if (TargetPin)
+			{
+				UEdGraphPin* TargetSource = ResolveLispExpression(Form->Get(1), Ctx);
+				if (TargetSource)
+				{
+					ConnectNodes(TargetSource, TargetPin, Ctx);
+				}
+			}
+			
+			// Connect arguments using SetPinValueFromExpr (handles assets, literals, etc.)
+			int32 ArgIndex = 0;
+			for (int32 i = 3; i < Form->Num(); i++)
+			{
+				if (Form->Get(i)->IsKeyword()) continue;
+				
+				int32 PinArgIndex = 0;
+				for (UEdGraphPin* Pin : CallNode->Pins)
+				{
+					if (Pin->Direction == EGPD_Input && 
+						Pin->PinType.PinCategory != UEdGraphSchema_K2::PC_Exec &&
+						Pin->PinName != UEdGraphSchema_K2::PN_Self &&
+						!Pin->bHidden)
+					{
+						if (PinArgIndex == ArgIndex)
+						{
+							SetPinValueFromExpr(Pin, Form->Get(i), Ctx);
+							ArgIndex++;
+							break;
+						}
+						PinArgIndex++;
+					}
+				}
+			}
+			
+			OutExecPin = GetNodeExecOutput(CallNode);
+			Ctx.Warnings.Add(FString::Printf(TEXT("DEBUG call: OutExecPin=%s"), OutExecPin ? TEXT("SET") : TEXT("NULL")));
+			return CallNode;
+		}
+	}
+	
+	// Handle (foreach item array body...)
+	if (FormName.Equals(TEXT("foreach"), ESearchCase::IgnoreCase) && Form->Num() >= 4)
+	{
+		// Load the ForEachLoop macro from StandardMacros
+		UBlueprint* MacroLibrary = LoadObject<UBlueprint>(nullptr, TEXT("/Engine/EditorBlueprintResources/StandardMacros.StandardMacros"));
+		if (MacroLibrary)
+		{
+			UEdGraph* ForEachMacro = nullptr;
+			for (UEdGraph* MacroGraph : MacroLibrary->MacroGraphs)
+			{
+				if (MacroGraph && MacroGraph->GetName().Contains(TEXT("ForEachLoop")))
+				{
+					ForEachMacro = MacroGraph;
+					break;
+				}
+			}
+			
+			if (ForEachMacro)
+			{
+				UK2Node_MacroInstance* ForEachNode = NewObject<UK2Node_MacroInstance>(Ctx.Graph);
+				ForEachNode->SetMacroGraph(ForEachMacro);
+				ForEachNode->NodePosX = Ctx.CurrentX;
+				ForEachNode->NodePosY = Ctx.CurrentY;
+				Ctx.Graph->AddNode(ForEachNode, false, false);
+				ForEachNode->AllocateDefaultPins();
+				Ctx.AdvancePosition();
+				
+				// Connect array input
+				UEdGraphPin* ArrayPin = nullptr;
+				for (UEdGraphPin* Pin : ForEachNode->Pins)
+				{
+					if (Pin->PinName.ToString().Contains(TEXT("Array")) && Pin->Direction == EGPD_Input &&
+						Pin->PinType.PinCategory != UEdGraphSchema_K2::PC_Exec)
+					{
+						ArrayPin = Pin;
+						break;
+					}
+				}
+				
+				if (ArrayPin)
+				{
+					UEdGraphPin* ArraySource = ResolveLispExpression(Form->Get(2), Ctx);
+					if (ArraySource)
+					{
+						ConnectNodes(ArraySource, ArrayPin, Ctx);
+					}
+				}
+				
+				// Store the element variable reference for the loop body
+				FString ItemVarName = Form->Get(1)->IsSymbol() ? Form->Get(1)->StringValue : TEXT("item");
+				for (UEdGraphPin* Pin : ForEachNode->Pins)
+				{
+					if (Pin->PinName.ToString().Contains(TEXT("Element")) && Pin->Direction == EGPD_Output &&
+						Pin->PinType.PinCategory != UEdGraphSchema_K2::PC_Exec)
+					{
+						Ctx.VariableToNodeId.Add(ItemVarName, ForEachNode->NodeGuid.ToString());
+						Ctx.VariableToPin.Add(ItemVarName, Pin->PinName.ToString());
+						break;
+					}
+				}
+				
+				// Find loop body exec pin and convert body
+				UEdGraphPin* LoopBodyPin = nullptr;
+				for (UEdGraphPin* Pin : ForEachNode->Pins)
+				{
+					if (Pin->PinName.ToString().Contains(TEXT("Loop")) && Pin->Direction == EGPD_Output &&
+						Pin->PinType.PinCategory == UEdGraphSchema_K2::PC_Exec)
+					{
+						LoopBodyPin = Pin;
+						break;
+					}
+				}
+				
+				if (LoopBodyPin && Form->Num() > 3)
+				{
+					ConvertExecBody(Form->Get(3), Ctx, LoopBodyPin);
+				}
+				
+				// Return the "Completed" exec pin
+				OutExecPin = nullptr;
+				for (UEdGraphPin* Pin : ForEachNode->Pins)
+				{
+					if (Pin->PinName.ToString().Contains(TEXT("Completed")) && Pin->Direction == EGPD_Output &&
+						Pin->PinType.PinCategory == UEdGraphSchema_K2::PC_Exec)
+					{
+						OutExecPin = Pin;
+						break;
+					}
+				}
+				
+				return ForEachNode;
+			}
+		}
+		
+		Ctx.Warnings.Add(TEXT("Failed to create ForEachLoop - StandardMacros not found"));
+		return nullptr;
+	}
+	
+	// Handle (for index start end body...) or (forloop index start end body...) - integer loop
+	if ((FormName.Equals(TEXT("for"), ESearchCase::IgnoreCase) || 
+	     FormName.Equals(TEXT("forloop"), ESearchCase::IgnoreCase)) && Form->Num() >= 5)
+	{
+		// Load the ForLoop macro from StandardMacros
+		UBlueprint* MacroLibrary = LoadObject<UBlueprint>(nullptr, TEXT("/Engine/EditorBlueprintResources/StandardMacros.StandardMacros"));
+		if (MacroLibrary)
+		{
+			UEdGraph* ForLoopMacro = nullptr;
+			for (UEdGraph* MacroGraph : MacroLibrary->MacroGraphs)
+			{
+				// Match "ForLoop" but not "ForLoopWithBreak" (exact match preferred)
+				if (MacroGraph && MacroGraph->GetName().Equals(TEXT("ForLoop"), ESearchCase::IgnoreCase))
+				{
+					ForLoopMacro = MacroGraph;
+					break;
+				}
+			}
+			
+			if (ForLoopMacro)
+			{
+				UK2Node_MacroInstance* ForLoopNode = NewObject<UK2Node_MacroInstance>(Ctx.Graph);
+				ForLoopNode->SetMacroGraph(ForLoopMacro);
+				ForLoopNode->NodePosX = Ctx.CurrentX;
+				ForLoopNode->NodePosY = Ctx.CurrentY;
+				Ctx.Graph->AddNode(ForLoopNode, false, false);
+				ForLoopNode->AllocateDefaultPins();
+				Ctx.AdvancePosition();
+				
+				// Get the index variable name (arg 1)
+				FString IndexVarName = Form->Get(1)->IsSymbol() ? Form->Get(1)->StringValue : TEXT("Index");
+				
+				// Connect FirstIndex input (arg 2)
+				UEdGraphPin* FirstIndexPin = nullptr;
+				for (UEdGraphPin* Pin : ForLoopNode->Pins)
+				{
+					if (Pin->PinName.ToString().Contains(TEXT("FirstIndex")) && Pin->Direction == EGPD_Input &&
+						Pin->PinType.PinCategory != UEdGraphSchema_K2::PC_Exec)
+					{
+						FirstIndexPin = Pin;
+						break;
+					}
+				}
+				if (FirstIndexPin)
+				{
+					if (Form->Get(2)->IsNumber())
+					{
+						FirstIndexPin->DefaultValue = FString::FromInt((int32)Form->Get(2)->NumberValue);
+					}
+					else
+					{
+						UEdGraphPin* StartSource = ResolveLispExpression(Form->Get(2), Ctx);
+						if (StartSource)
+						{
+							ConnectNodes(StartSource, FirstIndexPin, Ctx);
+						}
+					}
+				}
+				
+				// Connect LastIndex input (arg 3)
+				UEdGraphPin* LastIndexPin = nullptr;
+				for (UEdGraphPin* Pin : ForLoopNode->Pins)
+				{
+					if (Pin->PinName.ToString().Contains(TEXT("LastIndex")) && Pin->Direction == EGPD_Input &&
+						Pin->PinType.PinCategory != UEdGraphSchema_K2::PC_Exec)
+					{
+						LastIndexPin = Pin;
+						break;
+					}
+				}
+				if (LastIndexPin)
+				{
+					if (Form->Get(3)->IsNumber())
+					{
+						LastIndexPin->DefaultValue = FString::FromInt((int32)Form->Get(3)->NumberValue);
+					}
+					else
+					{
+						UEdGraphPin* EndSource = ResolveLispExpression(Form->Get(3), Ctx);
+						if (EndSource)
+						{
+							ConnectNodes(EndSource, LastIndexPin, Ctx);
+						}
+					}
+				}
+				
+				// Store the index variable reference for the loop body (the "Index" output pin)
+				for (UEdGraphPin* Pin : ForLoopNode->Pins)
+				{
+					if (Pin->PinName.ToString().Equals(TEXT("Index")) && Pin->Direction == EGPD_Output &&
+						Pin->PinType.PinCategory != UEdGraphSchema_K2::PC_Exec)
+					{
+						Ctx.VariableToNodeId.Add(IndexVarName, ForLoopNode->NodeGuid.ToString());
+						Ctx.VariableToPin.Add(IndexVarName, Pin->PinName.ToString());
+						break;
+					}
+				}
+				
+				// Find loop body exec pin and convert body
+				UEdGraphPin* LoopBodyPin = nullptr;
+				for (UEdGraphPin* Pin : ForLoopNode->Pins)
+				{
+					if (Pin->PinName.ToString().Contains(TEXT("LoopBody")) && Pin->Direction == EGPD_Output &&
+						Pin->PinType.PinCategory == UEdGraphSchema_K2::PC_Exec)
+					{
+						LoopBodyPin = Pin;
+						break;
+					}
+				}
+				if (LoopBodyPin && Form->Num() > 4)
+				{
+					ConvertExecBody(Form->Get(4), Ctx, LoopBodyPin);
+				}
+				
+				// Return the "Completed" exec pin for chaining
+				OutExecPin = nullptr;
+				for (UEdGraphPin* Pin : ForLoopNode->Pins)
+				{
+					if (Pin->PinName.ToString().Contains(TEXT("Completed")) && Pin->Direction == EGPD_Output &&
+						Pin->PinType.PinCategory == UEdGraphSchema_K2::PC_Exec)
+					{
+						OutExecPin = Pin;
+						break;
+					}
+				}
+				
+				return ForLoopNode;
+			}
+		}
+		
+		Ctx.Warnings.Add(TEXT("Failed to create ForLoop - StandardMacros not found"));
+		return nullptr;
+	}
+	
+	// Handle (cast TargetType ObjectExpr body...)
+	if (FormName.Equals(TEXT("cast"), ESearchCase::IgnoreCase) && Form->Num() >= 3)
+	{
+		FString TargetClassName = Form->Get(1)->IsSymbol() ? Form->Get(1)->StringValue : TEXT("");
+		
+		// Find the target class
+		UClass* TargetClass = nullptr;
+		
+		// Try as a blueprint first
+		UBlueprint* TargetBlueprint = LoadObject<UBlueprint>(nullptr, *TargetClassName);
+		if (TargetBlueprint && TargetBlueprint->GeneratedClass)
+		{
+			TargetClass = TargetBlueprint->GeneratedClass;
+		}
+		else
+		{
+			// Try as a native class
+			TArray<FString> ModulesToSearch = {
+				TEXT("/Script/Engine."),
+				TEXT("/Script/CoreUObject."),
+				TEXT("/Script/UMG."),
+				TEXT("/Script/AIModule."),
+				TEXT("")
+			};
+			
+			for (const FString& Module : ModulesToSearch)
+			{
+				TargetClass = FindObject<UClass>(nullptr, *(Module + TargetClassName));
+				if (TargetClass) break;
+			}
+			
+			// Search all classes by name
+			if (!TargetClass)
+			{
+				for (TObjectIterator<UClass> It; It; ++It)
+				{
+					if (It->GetName().Equals(TargetClassName, ESearchCase::IgnoreCase))
+					{
+						TargetClass = *It;
+						break;
+					}
+				}
+			}
+		}
+		
+		if (TargetClass)
+		{
+			UK2Node_DynamicCast* CastNode = NewObject<UK2Node_DynamicCast>(Ctx.Graph);
+			CastNode->TargetType = TargetClass;
+			CastNode->NodePosX = Ctx.CurrentX;
+			CastNode->NodePosY = Ctx.CurrentY;
+			Ctx.Graph->AddNode(CastNode, false, false);
+			CastNode->AllocateDefaultPins();
+			Ctx.AdvancePosition();
+			
+			// Connect object input
+			UEdGraphPin* ObjectPin = CastNode->GetCastSourcePin();
+			if (ObjectPin && Form->Num() > 2)
+			{
+				UEdGraphPin* ObjectSource = ResolveLispExpression(Form->Get(2), Ctx);
+				if (ObjectSource)
+				{
+					ConnectNodes(ObjectSource, ObjectPin, Ctx);
+				}
+			}
+			
+			// IMPORTANT: Reconstruct the node after connecting to wildcard pins
+			// This resolves the types and makes the cast compile correctly
+			CastNode->ReconstructNode();
+			
+			// Store the cast result for the body to use
+			FString CastResultVar = TEXT("_cast_result");
+			UEdGraphPin* CastResultPin = CastNode->GetCastResultPin();
+			if (CastResultPin)
+			{
+				Ctx.VariableToNodeId.Add(CastResultVar, CastNode->NodeGuid.ToString());
+				Ctx.VariableToPin.Add(CastResultVar, CastResultPin->PinName.ToString());
+			}
+			
+			// Convert the success body
+			UEdGraphPin* SuccessPin = CastNode->GetValidCastPin();
+			if (SuccessPin && Form->Num() > 3)
+			{
+				ConvertExecBody(Form->Get(3), Ctx, SuccessPin);
+			}
+			
+			// Return the invalid cast pin as the continuation
+			OutExecPin = CastNode->GetInvalidCastPin();
+			return CastNode;
+		}
+		else
+		{
+			Ctx.Warnings.Add(FString::Printf(TEXT("Cast target class not found: %s"), *TargetClassName));
+		}
+	}
+	
+	// Handle (print message) as shorthand for PrintString
+	if (FormName.Equals(TEXT("print"), ESearchCase::IgnoreCase) && Form->Num() >= 2)
+	{
+		UFunction* PrintFunc = UKismetSystemLibrary::StaticClass()->FindFunctionByName(TEXT("PrintString"));
+		if (PrintFunc)
+		{
+			UK2Node_CallFunction* PrintNode = NewObject<UK2Node_CallFunction>(Ctx.Graph);
+			PrintNode->SetFromFunction(PrintFunc);
+			PrintNode->NodePosX = Ctx.CurrentX;
+			PrintNode->NodePosY = Ctx.CurrentY;
+			Ctx.Graph->AddNode(PrintNode, false, false);
+			PrintNode->AllocateDefaultPins();
+			EnsureNodeHasValidGuid(PrintNode);
+			Ctx.AdvancePosition();
+			
+			// Set the string to print
+			UEdGraphPin* StringPin = PrintNode->FindPin(TEXT("InString"));
+			if (StringPin)
+			{
+				if (Form->Get(1)->IsString())
+				{
+					StringPin->DefaultValue = Form->Get(1)->StringValue;
+				}
+				else if (Form->Get(1)->IsNumber())
+				{
+					// Convert number to string
+					SetNumericPinDefaultValue(StringPin, Form->Get(1)->NumberValue);
+				}
+				else
+				{
+					UEdGraphPin* ValueSource = ResolveLispExpression(Form->Get(1), Ctx);
+					if (ValueSource)
+					{
+						// Check if we need to convert to string
+						if (ValueSource->PinType.PinCategory != UEdGraphSchema_K2::PC_String)
+						{
+							// Create a conversion node based on the source type
+							FString ConvFuncName;
+							if (ValueSource->PinType.PinCategory == UEdGraphSchema_K2::PC_Int)
+							{
+								ConvFuncName = TEXT("Conv_IntToString");
+							}
+							else if (ValueSource->PinType.PinCategory == UEdGraphSchema_K2::PC_Real ||
+									 ValueSource->PinType.PinCategory == UEdGraphSchema_K2::PC_Float ||
+									 ValueSource->PinType.PinCategory == UEdGraphSchema_K2::PC_Double)
+							{
+								ConvFuncName = TEXT("Conv_DoubleToString");
+								// Try float version as fallback
+								if (!UKismetStringLibrary::StaticClass()->FindFunctionByName(*ConvFuncName))
+								{
+									ConvFuncName = TEXT("Conv_FloatToString");
+								}
+							}
+							else if (ValueSource->PinType.PinCategory == UEdGraphSchema_K2::PC_Boolean)
+							{
+								ConvFuncName = TEXT("Conv_BoolToString");
+							}
+							else if (ValueSource->PinType.PinCategory == UEdGraphSchema_K2::PC_Name)
+							{
+								ConvFuncName = TEXT("Conv_NameToString");
+							}
+							else if (ValueSource->PinType.PinCategory == UEdGraphSchema_K2::PC_Text)
+							{
+								ConvFuncName = TEXT("Conv_TextToString");
+							}
+							else if (ValueSource->PinType.PinCategory == UEdGraphSchema_K2::PC_Struct)
+							{
+								// For vectors, rotators, etc.
+								UScriptStruct* StructType = Cast<UScriptStruct>(ValueSource->PinType.PinSubCategoryObject.Get());
+								if (StructType)
+								{
+									FString StructName = StructType->GetName();
+									if (StructName == TEXT("Vector"))
+									{
+										ConvFuncName = TEXT("Conv_VectorToString");
+									}
+									else if (StructName == TEXT("Rotator"))
+									{
+										ConvFuncName = TEXT("Conv_RotatorToString");
+									}
+									else if (StructName == TEXT("Transform"))
+									{
+										ConvFuncName = TEXT("Conv_TransformToString");
+									}
+								}
+							}
+							
+							if (!ConvFuncName.IsEmpty())
+							{
+								UFunction* ConvFunc = UKismetStringLibrary::StaticClass()->FindFunctionByName(*ConvFuncName);
+								if (ConvFunc)
+								{
+									UK2Node_CallFunction* ConvNode = NewObject<UK2Node_CallFunction>(Ctx.Graph);
+									ConvNode->SetFromFunction(ConvFunc);
+									ConvNode->NodePosX = Ctx.CurrentX - 200;
+									ConvNode->NodePosY = Ctx.CurrentY;
+									Ctx.Graph->AddNode(ConvNode, false, false);
+									ConvNode->AllocateDefaultPins();
+									EnsureNodeHasValidGuid(ConvNode);
+									
+									// Connect source to conversion node input
+									for (UEdGraphPin* Pin : ConvNode->Pins)
+									{
+										if (Pin->Direction == EGPD_Input && 
+											Pin->PinType.PinCategory != UEdGraphSchema_K2::PC_Exec)
+										{
+											ConnectNodes(ValueSource, Pin, Ctx);
+											break;
+										}
+									}
+									
+									// Connect conversion output to print string input
+									UEdGraphPin* ConvOutput = FindOutputPin(ConvNode, TEXT("ReturnValue"));
+									if (ConvOutput)
+									{
+										ConnectNodes(ConvOutput, StringPin, Ctx);
+									}
+								}
+								else
+								{
+									// Couldn't find conversion, try direct connection
+									ConnectNodes(ValueSource, StringPin, Ctx);
+								}
+							}
+							else
+							{
+								// No known conversion, try direct connection
+								ConnectNodes(ValueSource, StringPin, Ctx);
+							}
+						}
+						else
+						{
+							// Already a string, connect directly
+							ConnectNodes(ValueSource, StringPin, Ctx);
+						}
+					}
+				}
+			}
+			
+			OutExecPin = GetNodeExecOutput(PrintNode);
+			return PrintNode;
+		}
+	}
+	
+	// Handle (delay seconds)
+	if (FormName.Equals(TEXT("delay"), ESearchCase::IgnoreCase))
+	{
+		// Delay is a latent function on UKismetSystemLibrary
+		UFunction* DelayFunc = UKismetSystemLibrary::StaticClass()->FindFunctionByName(TEXT("Delay"));
+		if (DelayFunc)
+		{
+			UK2Node_CallFunction* DelayNode = NewObject<UK2Node_CallFunction>(Ctx.Graph);
+			DelayNode->SetFromFunction(DelayFunc);
+			DelayNode->NodePosX = Ctx.CurrentX;
+			DelayNode->NodePosY = Ctx.CurrentY;
+			Ctx.Graph->AddNode(DelayNode, false, false);
+			DelayNode->AllocateDefaultPins();
+			Ctx.AdvancePosition();
+			
+			// Set duration
+			if (Form->Num() > 1)
+			{
+				UEdGraphPin* DurationPin = DelayNode->FindPin(TEXT("Duration"));
+				if (DurationPin)
+				{
+					if (Form->Get(1)->IsNumber())
+					{
+						SetNumericPinDefaultValue(DurationPin, Form->Get(1)->NumberValue);
+					}
+					else
+					{
+						UEdGraphPin* DurationSource = ResolveLispExpression(Form->Get(1), Ctx);
+						if (DurationSource)
+						{
+							ConnectNodes(DurationSource, DurationPin, Ctx);
+						}
+					}
+				}
+			}
+			
+			OutExecPin = GetNodeExecOutput(DelayNode);
+			return DelayNode;
+		}
+	}
+	
+	// Handle (switch-int value :0 body0 :1 body1 :default defaultBody)
+	if (FormName.Equals(TEXT("switch-int"), ESearchCase::IgnoreCase) && Form->Num() >= 2)
+	{
+		UK2Node_SwitchInteger* SwitchNode = NewObject<UK2Node_SwitchInteger>(Ctx.Graph);
+		SwitchNode->NodePosX = Ctx.CurrentX;
+		SwitchNode->NodePosY = Ctx.CurrentY;
+		Ctx.Graph->AddNode(SwitchNode, false, false);
+		SwitchNode->AllocateDefaultPins();
+		Ctx.AdvancePosition();
+		
+		// Connect selection value
+		UEdGraphPin* SelectionPin = SwitchNode->GetSelectionPin();
+		if (SelectionPin)
+		{
+			UEdGraphPin* ValueSource = ResolveLispExpression(Form->Get(1), Ctx);
+			if (ValueSource)
+			{
+				ConnectNodes(ValueSource, SelectionPin, Ctx);
+			}
+			else if (Form->Get(1)->IsNumber())
+			{
+				SelectionPin->DefaultValue = FString::Printf(TEXT("%d"), (int32)Form->Get(1)->NumberValue);
+			}
+		}
+		
+		// Process case bodies - keywords are case values
+		for (int32 i = 2; i < Form->Num() - 1; i += 2)
+		{
+			if (!Form->Get(i)->IsKeyword()) continue;
+			
+			FString CaseKey = Form->Get(i)->StringValue;
+			CaseKey.RemoveFromStart(TEXT(":"));
+			
+			if (CaseKey.Equals(TEXT("default"), ESearchCase::IgnoreCase))
+			{
+				UEdGraphPin* DefaultPin = SwitchNode->GetDefaultPin();
+				if (DefaultPin && i + 1 < Form->Num())
+				{
+					ConvertExecBody(Form->Get(i + 1), Ctx, DefaultPin);
+				}
+			}
+			else if (CaseKey.IsNumeric())
+			{
+				int32 CaseValue = FCString::Atoi(*CaseKey);
+				// Add pin for this case if needed
+				SwitchNode->AddPinToSwitchNode();
+				
+				// Find the pin for this case value
+				for (UEdGraphPin* Pin : SwitchNode->Pins)
+				{
+					if (Pin->Direction == EGPD_Output && 
+						Pin->PinType.PinCategory == UEdGraphSchema_K2::PC_Exec &&
+						Pin->PinName.ToString() == CaseKey)
+					{
+						if (i + 1 < Form->Num())
+						{
+							UEdGraphPin* CasePin = Pin;
+							ConvertExecBody(Form->Get(i + 1), Ctx, CasePin);
+						}
+						break;
+					}
+				}
+			}
+		}
+		
+		// Switch nodes don't have a single "then" pin - execution continues from each case
+		OutExecPin = nullptr;
+		return SwitchNode;
+	}
+	
+	// Handle (comment "text") - creates a comment node
+	if (FormName.Equals(TEXT("comment"), ESearchCase::IgnoreCase) && Form->Num() >= 2)
+	{
+		UEdGraphNode_Comment* CommentNode = NewObject<UEdGraphNode_Comment>(Ctx.Graph);
+		CommentNode->NodePosX = Ctx.CurrentX;
+		CommentNode->NodePosY = Ctx.CurrentY - 100; // Place above current position
+		
+		FString CommentText = Form->Get(1)->IsString() ? Form->Get(1)->StringValue : Form->Get(1)->ToString(false);
+		CommentNode->NodeComment = CommentText;
+		
+		// Set size based on text length
+		CommentNode->NodeWidth = FMath::Max(200, CommentText.Len() * 8);
+		CommentNode->NodeHeight = 100;
+		
+		Ctx.Graph->AddNode(CommentNode, false, false);
+		
+		// Comments don't affect execution flow
+		OutExecPin = nullptr;
+		return nullptr; // Don't chain execution through comments
+	}
+	
+	// Handle (on-component ComponentName EventName :params (...) body...)
+	// Example: (on-component BoxCollision BeginOverlap :params ((OtherActor Actor)) (print "Overlap!"))
+	if (FormName.Equals(TEXT("on-component"), ESearchCase::IgnoreCase) && Form->Num() >= 3)
+	{
+		FString ComponentName = Form->Get(1)->IsSymbol() ? Form->Get(1)->StringValue : TEXT("");
+		FString EventName = Form->Get(2)->IsSymbol() ? Form->Get(2)->StringValue : TEXT("");
+		
+		// Find the component property in the Blueprint
+		FObjectProperty* ComponentProperty = nullptr;
+		if (Ctx.Blueprint->SimpleConstructionScript)
+		{
+			for (USCS_Node* SCSNode : Ctx.Blueprint->SimpleConstructionScript->GetAllNodes())
+			{
+				if (SCSNode && SCSNode->GetVariableName().ToString().Equals(ComponentName, ESearchCase::IgnoreCase))
+				{
+					// Found the component - now find the property
+					ComponentProperty = FindFProperty<FObjectProperty>(Ctx.Blueprint->SkeletonGeneratedClass, SCSNode->GetVariableName());
+					break;
+				}
+			}
+		}
+		
+		if (!ComponentProperty)
+		{
+			Ctx.Warnings.Add(FString::Printf(TEXT("Component not found: %s"), *ComponentName));
+		}
+		
+		// Build the delegate name (e.g., "BeginOverlap" -> "OnComponentBeginOverlap")
+		FString DelegateName = TEXT("OnComponent") + EventName;
+		if (!EventName.StartsWith(TEXT("On")))
+		{
+			// Also try just "On" + EventName
+		}
+		
+		// Find the delegate property
+		FMulticastDelegateProperty* DelegateProperty = nullptr;
+		if (ComponentProperty)
+		{
+			UClass* ComponentClass = ComponentProperty->PropertyClass;
+			if (ComponentClass)
+			{
+				DelegateProperty = FindFProperty<FMulticastDelegateProperty>(ComponentClass, *DelegateName);
+				if (!DelegateProperty)
+				{
+					// Try without "Component" prefix
+					DelegateName = TEXT("On") + EventName;
+					DelegateProperty = FindFProperty<FMulticastDelegateProperty>(ComponentClass, *DelegateName);
+				}
+			}
+		}
+		
+		// Create the component bound event node
+		UK2Node_ComponentBoundEvent* CompEventNode = NewObject<UK2Node_ComponentBoundEvent>(Ctx.Graph);
+		if (ComponentProperty && DelegateProperty)
+		{
+			CompEventNode->InitializeComponentBoundEventParams(ComponentProperty, DelegateProperty);
+		}
+		else
+		{
+			// Manual setup if we couldn't find the properties
+			CompEventNode->ComponentPropertyName = FName(*ComponentName);
+			CompEventNode->DelegatePropertyName = FName(*DelegateName);
+		}
+		
+		CompEventNode->NodePosX = Ctx.CurrentX;
+		CompEventNode->NodePosY = Ctx.CurrentY;
+		Ctx.Graph->AddNode(CompEventNode, false, false);
+		CompEventNode->AllocateDefaultPins();
+		Ctx.NewRow();
+		
+		// Add the event node to TempIdToNode so variable lookups can find it
+		FString EventTempId = Ctx.GenerateTempId();
+		Ctx.TempIdToNode.Add(EventTempId, CompEventNode);
+		
+		// Register all output pins as variables so they can be referenced in the body
+		// This allows (cast Character OtherActor ...) to find OtherActor
+		for (UEdGraphPin* Pin : CompEventNode->Pins)
+		{
+			if (Pin->Direction == EGPD_Output && 
+				Pin->PinType.PinCategory != UEdGraphSchema_K2::PC_Exec &&
+				!Pin->bHidden)
+			{
+				FString PinName = Pin->PinName.ToString();
+				Ctx.VariableToNodeId.Add(PinName, CompEventNode->NodeGuid.ToString());
+				Ctx.VariableToPin.Add(PinName, PinName);
+				
+				// Also register with common variations
+				// e.g., "Other Actor" -> "OtherActor"
+				FString NoSpaces = PinName.Replace(TEXT(" "), TEXT(""));
+				if (NoSpaces != PinName)
+				{
+					Ctx.VariableToNodeId.Add(NoSpaces, CompEventNode->NodeGuid.ToString());
+					Ctx.VariableToPin.Add(NoSpaces, PinName);
+				}
+			}
+		}
+		
+		// Find body start (after :params if present)
+		int32 BodyStart = 3;
+		for (int32 i = 3; i < Form->Num(); i++)
+		{
+			if (Form->Get(i)->IsKeyword() && Form->Get(i)->StringValue == TEXT(":params"))
+			{
+				BodyStart = i + 2; // Skip :params and the params list
+				break;
+			}
+		}
+		
+		// Convert the body
+		UEdGraphPin* CurrentExecPin = GetNodeExecOutput(CompEventNode);
+		for (int32 i = BodyStart; i < Form->Num(); i++)
+		{
+			if (Form->Get(i)->IsKeyword()) continue;
+			
+			UEdGraphPin* NodeExecOut = nullptr;
+			UEdGraphNode* Node = ConvertLispFormToNode(Form->Get(i), Ctx, NodeExecOut);
+			
+			if (Node && CurrentExecPin)
+			{
+				UEdGraphPin* NodeInput = GetNodeExecInput(Node);
+				if (NodeInput)
+				{
+					ConnectNodes(CurrentExecPin, NodeInput, Ctx);
+				}
+			}
+			
+			if (NodeExecOut)
+			{
+				CurrentExecPin = NodeExecOut;
+			}
+		}
+		
+		OutExecPin = nullptr; // Event nodes don't chain
+		return CompEventNode;
+	}
+	
+	// Handle (on-input ActionName :pressed body :released body)
+	if (FormName.Equals(TEXT("on-input"), ESearchCase::IgnoreCase) && Form->Num() >= 2)
+	{
+		FString ActionName = Form->Get(1)->IsSymbol() ? Form->Get(1)->StringValue : TEXT("");
+		
+		UK2Node_InputAction* InputNode = NewObject<UK2Node_InputAction>(Ctx.Graph);
+		InputNode->InputActionName = FName(*ActionName);
+		InputNode->NodePosX = Ctx.CurrentX;
+		InputNode->NodePosY = Ctx.CurrentY;
+		Ctx.Graph->AddNode(InputNode, false, false);
+		InputNode->AllocateDefaultPins();
+		Ctx.AdvancePosition();
+		
+		// Process event type bodies
+		for (int32 i = 2; i < Form->Num() - 1; i += 2)
+		{
+			if (!Form->Get(i)->IsKeyword()) continue;
+			
+			FString EventType = Form->Get(i)->StringValue;
+			EventType.RemoveFromStart(TEXT(":"));
+			
+			// Find the matching exec pin
+			for (UEdGraphPin* Pin : InputNode->Pins)
+			{
+				if (Pin->Direction == EGPD_Output && 
+					Pin->PinType.PinCategory == UEdGraphSchema_K2::PC_Exec &&
+					Pin->PinName.ToString().Equals(EventType, ESearchCase::IgnoreCase))
+				{
+					if (i + 1 < Form->Num())
+					{
+						UEdGraphPin* EventPin = Pin;
+						ConvertExecBody(Form->Get(i + 1), Ctx, EventPin);
+					}
+					break;
+				}
+			}
+		}
+		
+		OutExecPin = nullptr; // Input action is an event, not a statement
+		return InputNode;
+	}
+	
+	// Handle (spawn ActorClass Location) - shorthand for SpawnActor
+	if (FormName.Equals(TEXT("spawn"), ESearchCase::IgnoreCase) && Form->Num() >= 2)
+	{
+		UFunction* SpawnFunc = UGameplayStatics::StaticClass()->FindFunctionByName(TEXT("BeginDeferredActorSpawnFromClass"));
+		if (SpawnFunc)
+		{
+			UK2Node_CallFunction* SpawnNode = NewObject<UK2Node_CallFunction>(Ctx.Graph);
+			SpawnNode->SetFromFunction(SpawnFunc);
+			SpawnNode->NodePosX = Ctx.CurrentX;
+			SpawnNode->NodePosY = Ctx.CurrentY;
+			Ctx.Graph->AddNode(SpawnNode, false, false);
+			SpawnNode->AllocateDefaultPins();
+			Ctx.AdvancePosition();
+			
+			// Set class - the first argument should be a class reference
+			// This is complex because we need a class pin, for now just warn
+			Ctx.Warnings.Add(TEXT("spawn shorthand created node but class pin needs manual connection"));
+			
+			OutExecPin = GetNodeExecOutput(SpawnNode);
+			return SpawnNode;
+		}
+	}
+	
+	// Handle (destroy) or (destroy target) - shorthand for DestroyActor
+	if (FormName.Equals(TEXT("destroy"), ESearchCase::IgnoreCase))
+	{
+		UFunction* DestroyFunc = AActor::StaticClass()->FindFunctionByName(TEXT("K2_DestroyActor"));
+		if (DestroyFunc)
+		{
+			UK2Node_CallFunction* DestroyNode = NewObject<UK2Node_CallFunction>(Ctx.Graph);
+			DestroyNode->SetFromFunction(DestroyFunc);
+			DestroyNode->NodePosX = Ctx.CurrentX;
+			DestroyNode->NodePosY = Ctx.CurrentY;
+			Ctx.Graph->AddNode(DestroyNode, false, false);
+			DestroyNode->AllocateDefaultPins();
+			Ctx.AdvancePosition();
+			
+			// If there's a target, connect it
+			if (Form->Num() > 1)
+			{
+				UEdGraphPin* TargetPin = DestroyNode->FindPin(UEdGraphSchema_K2::PN_Self);
+				if (TargetPin)
+				{
+					UEdGraphPin* TargetSource = ResolveLispExpression(Form->Get(1), Ctx);
+					if (TargetSource)
+					{
+						ConnectNodes(TargetSource, TargetPin, Ctx);
+					}
+				}
+			}
+			
+			OutExecPin = GetNodeExecOutput(DestroyNode);
+			return DestroyNode;
+		}
+	}
+	
+	// Handle (get-component ComponentClass) - shorthand for GetComponentByClass
+	if (FormName.Equals(TEXT("get-component"), ESearchCase::IgnoreCase) && Form->Num() >= 2)
+	{
+		UFunction* GetCompFunc = AActor::StaticClass()->FindFunctionByName(TEXT("GetComponentByClass"));
+		if (GetCompFunc)
+		{
+			UK2Node_CallFunction* GetCompNode = NewObject<UK2Node_CallFunction>(Ctx.Graph);
+			GetCompNode->SetFromFunction(GetCompFunc);
+			GetCompNode->NodePosX = Ctx.CurrentX;
+			GetCompNode->NodePosY = Ctx.CurrentY;
+			Ctx.Graph->AddNode(GetCompNode, false, false);
+			GetCompNode->AllocateDefaultPins();
+			
+			FString TempId = Ctx.GenerateTempId();
+			Ctx.TempIdToNode.Add(TempId, GetCompNode);
+			
+			// This is a pure node, no exec pin
+			OutExecPin = nullptr;
+			return GetCompNode;
+		}
+	}
+	
+	// Handle (play-sound SoundAsset) or (play-sound SoundAsset Location)
+	if (FormName.Equals(TEXT("play-sound"), ESearchCase::IgnoreCase) && Form->Num() >= 2)
+	{
+		UFunction* PlaySoundFunc = UGameplayStatics::StaticClass()->FindFunctionByName(TEXT("PlaySoundAtLocation"));
+		if (PlaySoundFunc)
+		{
+			UK2Node_CallFunction* PlayNode = NewObject<UK2Node_CallFunction>(Ctx.Graph);
+			PlayNode->SetFromFunction(PlaySoundFunc);
+			PlayNode->NodePosX = Ctx.CurrentX;
+			PlayNode->NodePosY = Ctx.CurrentY;
+			Ctx.Graph->AddNode(PlayNode, false, false);
+			PlayNode->AllocateDefaultPins();
+			Ctx.AdvancePosition();
+			
+			// Sound is a complex type - we'll create the node but warn about manual connection needed
+			Ctx.Warnings.Add(TEXT("play-sound created node - Sound and Location pins may need manual connection"));
+			
+			OutExecPin = GetNodeExecOutput(PlayNode);
+			return PlayNode;
+		}
+	}
+	
+	// Handle (set-timer Duration bLoop) - creates a timer
+	if (FormName.Equals(TEXT("set-timer"), ESearchCase::IgnoreCase) && Form->Num() >= 2)
+	{
+		UFunction* SetTimerFunc = UKismetSystemLibrary::StaticClass()->FindFunctionByName(TEXT("K2_SetTimer"));
+		if (SetTimerFunc)
+		{
+			UK2Node_CallFunction* TimerNode = NewObject<UK2Node_CallFunction>(Ctx.Graph);
+			TimerNode->SetFromFunction(SetTimerFunc);
+			TimerNode->NodePosX = Ctx.CurrentX;
+			TimerNode->NodePosY = Ctx.CurrentY;
+			Ctx.Graph->AddNode(TimerNode, false, false);
+			TimerNode->AllocateDefaultPins();
+			Ctx.AdvancePosition();
+			
+			// Set time
+			UEdGraphPin* TimePin = TimerNode->FindPin(TEXT("Time"));
+			if (TimePin && Form->Num() > 1)
+			{
+				if (Form->Get(1)->IsNumber())
+				{
+					SetNumericPinDefaultValue(TimePin, Form->Get(1)->NumberValue);
+				}
+			}
+			
+			// Set looping
+			UEdGraphPin* LoopPin = TimerNode->FindPin(TEXT("bLooping"));
+			if (LoopPin && Form->Num() > 2)
+			{
+				if (Form->Get(2)->IsSymbol())
+				{
+					LoopPin->DefaultValue = Form->Get(2)->StringValue.Equals(TEXT("true"), ESearchCase::IgnoreCase) ? TEXT("true") : TEXT("false");
+				}
+			}
+			
+			OutExecPin = GetNodeExecOutput(TimerNode);
+			return TimerNode;
+		}
+	}
+	
+	// Handle (clear-timer Handle) - clears a timer
+	if (FormName.Equals(TEXT("clear-timer"), ESearchCase::IgnoreCase))
+	{
+		UFunction* ClearTimerFunc = UKismetSystemLibrary::StaticClass()->FindFunctionByName(TEXT("K2_ClearTimer"));
+		if (ClearTimerFunc)
+		{
+			UK2Node_CallFunction* ClearNode = NewObject<UK2Node_CallFunction>(Ctx.Graph);
+			ClearNode->SetFromFunction(ClearTimerFunc);
+			ClearNode->NodePosX = Ctx.CurrentX;
+			ClearNode->NodePosY = Ctx.CurrentY;
+			Ctx.Graph->AddNode(ClearNode, false, false);
+			ClearNode->AllocateDefaultPins();
+			Ctx.AdvancePosition();
+			
+			OutExecPin = GetNodeExecOutput(ClearNode);
+			return ClearNode;
+		}
+	}
+	
+	// Handle (array-add array item) - add item to array
+	if (FormName.Equals(TEXT("array-add"), ESearchCase::IgnoreCase) && Form->Num() >= 3)
+	{
+		UFunction* AddFunc = UKismetArrayLibrary::StaticClass()->FindFunctionByName(TEXT("Array_Add"));
+		if (AddFunc)
+		{
+			UK2Node_CallFunction* AddNode = NewObject<UK2Node_CallFunction>(Ctx.Graph);
+			AddNode->SetFromFunction(AddFunc);
+			AddNode->NodePosX = Ctx.CurrentX;
+			AddNode->NodePosY = Ctx.CurrentY;
+			Ctx.Graph->AddNode(AddNode, false, false);
+			AddNode->AllocateDefaultPins();
+			Ctx.AdvancePosition();
+			
+			// Connect array
+			UEdGraphPin* ArrayPin = AddNode->FindPin(TEXT("TargetArray"));
+			if (ArrayPin)
+			{
+				UEdGraphPin* ArraySource = ResolveLispExpression(Form->Get(1), Ctx);
+				if (ArraySource)
+				{
+					ConnectNodes(ArraySource, ArrayPin, Ctx);
+				}
+			}
+			
+			// Connect item
+			UEdGraphPin* ItemPin = AddNode->FindPin(TEXT("NewItem"));
+			if (ItemPin)
+			{
+				UEdGraphPin* ItemSource = ResolveLispExpression(Form->Get(2), Ctx);
+				if (ItemSource)
+				{
+					ConnectNodes(ItemSource, ItemPin, Ctx);
+				}
+			}
+			
+			OutExecPin = GetNodeExecOutput(AddNode);
+			return AddNode;
+		}
+	}
+	
+	// Handle (array-remove array index) - remove item at index
+	if (FormName.Equals(TEXT("array-remove"), ESearchCase::IgnoreCase) && Form->Num() >= 3)
+	{
+		UFunction* RemoveFunc = UKismetArrayLibrary::StaticClass()->FindFunctionByName(TEXT("Array_Remove"));
+		if (RemoveFunc)
+		{
+			UK2Node_CallFunction* RemoveNode = NewObject<UK2Node_CallFunction>(Ctx.Graph);
+			RemoveNode->SetFromFunction(RemoveFunc);
+			RemoveNode->NodePosX = Ctx.CurrentX;
+			RemoveNode->NodePosY = Ctx.CurrentY;
+			Ctx.Graph->AddNode(RemoveNode, false, false);
+			RemoveNode->AllocateDefaultPins();
+			Ctx.AdvancePosition();
+			
+			UEdGraphPin* ArrayPin = RemoveNode->FindPin(TEXT("TargetArray"));
+			if (ArrayPin)
+			{
+				UEdGraphPin* ArraySource = ResolveLispExpression(Form->Get(1), Ctx);
+				if (ArraySource)
+				{
+					ConnectNodes(ArraySource, ArrayPin, Ctx);
+				}
+			}
+			
+			UEdGraphPin* IndexPin = RemoveNode->FindPin(TEXT("IndexToRemove"));
+			if (IndexPin)
+			{
+				if (Form->Get(2)->IsNumber())
+				{
+					IndexPin->DefaultValue = FString::Printf(TEXT("%d"), (int32)Form->Get(2)->NumberValue);
+				}
+				else
+				{
+					UEdGraphPin* IndexSource = ResolveLispExpression(Form->Get(2), Ctx);
+					if (IndexSource)
+					{
+						ConnectNodes(IndexSource, IndexPin, Ctx);
+					}
+				}
+			}
+			
+			OutExecPin = GetNodeExecOutput(RemoveNode);
+			return RemoveNode;
+		}
+	}
+	
+	// Handle (array-clear array) - clear array
+	if (FormName.Equals(TEXT("array-clear"), ESearchCase::IgnoreCase) && Form->Num() >= 2)
+	{
+		UFunction* ClearFunc = UKismetArrayLibrary::StaticClass()->FindFunctionByName(TEXT("Array_Clear"));
+		if (ClearFunc)
+		{
+			UK2Node_CallFunction* ClearNode = NewObject<UK2Node_CallFunction>(Ctx.Graph);
+			ClearNode->SetFromFunction(ClearFunc);
+			ClearNode->NodePosX = Ctx.CurrentX;
+			ClearNode->NodePosY = Ctx.CurrentY;
+			Ctx.Graph->AddNode(ClearNode, false, false);
+			ClearNode->AllocateDefaultPins();
+			Ctx.AdvancePosition();
+			
+			UEdGraphPin* ArrayPin = ClearNode->FindPin(TEXT("TargetArray"));
+			if (ArrayPin)
+			{
+				UEdGraphPin* ArraySource = ResolveLispExpression(Form->Get(1), Ctx);
+				if (ArraySource)
+				{
+					ConnectNodes(ArraySource, ArrayPin, Ctx);
+				}
+			}
+			
+			OutExecPin = GetNodeExecOutput(ClearNode);
+			return ClearNode;
+		}
+	}
+	
+	// Handle (array-shuffle array) - shuffle array
+	if (FormName.Equals(TEXT("array-shuffle"), ESearchCase::IgnoreCase) && Form->Num() >= 2)
+	{
+		UFunction* ShuffleFunc = UKismetArrayLibrary::StaticClass()->FindFunctionByName(TEXT("Array_Shuffle"));
+		if (ShuffleFunc)
+		{
+			UK2Node_CallFunction* ShuffleNode = NewObject<UK2Node_CallFunction>(Ctx.Graph);
+			ShuffleNode->SetFromFunction(ShuffleFunc);
+			ShuffleNode->NodePosX = Ctx.CurrentX;
+			ShuffleNode->NodePosY = Ctx.CurrentY;
+			Ctx.Graph->AddNode(ShuffleNode, false, false);
+			ShuffleNode->AllocateDefaultPins();
+			Ctx.AdvancePosition();
+			
+			UEdGraphPin* ArrayPin = ShuffleNode->FindPin(TEXT("TargetArray"));
+			if (ArrayPin)
+			{
+				UEdGraphPin* ArraySource = ResolveLispExpression(Form->Get(1), Ctx);
+				if (ArraySource)
+				{
+					ConnectNodes(ArraySource, ArrayPin, Ctx);
+				}
+			}
+			
+			OutExecPin = GetNodeExecOutput(ShuffleNode);
+			return ShuffleNode;
+		}
+	}
+	
+	// Handle (array-reverse array) - reverse array
+	if (FormName.Equals(TEXT("array-reverse"), ESearchCase::IgnoreCase) && Form->Num() >= 2)
+	{
+		UFunction* ReverseFunc = UKismetArrayLibrary::StaticClass()->FindFunctionByName(TEXT("Array_Reverse"));
+		if (ReverseFunc)
+		{
+			UK2Node_CallFunction* ReverseNode = NewObject<UK2Node_CallFunction>(Ctx.Graph);
+			ReverseNode->SetFromFunction(ReverseFunc);
+			ReverseNode->NodePosX = Ctx.CurrentX;
+			ReverseNode->NodePosY = Ctx.CurrentY;
+			Ctx.Graph->AddNode(ReverseNode, false, false);
+			ReverseNode->AllocateDefaultPins();
+			Ctx.AdvancePosition();
+			
+			UEdGraphPin* ArrayPin = ReverseNode->FindPin(TEXT("TargetArray"));
+			if (ArrayPin)
+			{
+				UEdGraphPin* ArraySource = ResolveLispExpression(Form->Get(1), Ctx);
+				if (ArraySource)
+				{
+					ConnectNodes(ArraySource, ArrayPin, Ctx);
+				}
+			}
+			
+			OutExecPin = GetNodeExecOutput(ReverseNode);
+			return ReverseNode;
+		}
+	}
+	
+	// Handle (log message) or (log message Category) - prints to log
+	if (FormName.Equals(TEXT("log"), ESearchCase::IgnoreCase) && Form->Num() >= 2)
+	{
+		UFunction* LogFunc = UKismetSystemLibrary::StaticClass()->FindFunctionByName(TEXT("PrintWarning"));
+		if (LogFunc)
+		{
+			UK2Node_CallFunction* LogNode = NewObject<UK2Node_CallFunction>(Ctx.Graph);
+			LogNode->SetFromFunction(LogFunc);
+			LogNode->NodePosX = Ctx.CurrentX;
+			LogNode->NodePosY = Ctx.CurrentY;
+			Ctx.Graph->AddNode(LogNode, false, false);
+			LogNode->AllocateDefaultPins();
+			Ctx.AdvancePosition();
+			
+			UEdGraphPin* StringPin = LogNode->FindPin(TEXT("InString"));
+			if (StringPin)
+			{
+				if (Form->Get(1)->IsString())
+				{
+					StringPin->DefaultValue = Form->Get(1)->StringValue;
+				}
+				else
+				{
+					UEdGraphPin* StringSource = ResolveLispExpression(Form->Get(1), Ctx);
+					if (StringSource)
+					{
+						ConnectNodes(StringSource, StringPin, Ctx);
+					}
+				}
+			}
+			
+			OutExecPin = GetNodeExecOutput(LogNode);
+			return LogNode;
+		}
+	}
+	
+	// Handle (valid? expr) - shorthand for IsValid
+	// Handle (return) or (return value) - return from function
+	if (FormName.Equals(TEXT("return"), ESearchCase::IgnoreCase))
+	{
+		// Find the function result node in this graph
+		UK2Node_FunctionResult* ResultNode = nullptr;
+		for (UEdGraphNode* Node : Ctx.Graph->Nodes)
+		{
+			if (UK2Node_FunctionResult* Result = Cast<UK2Node_FunctionResult>(Node))
+			{
+				ResultNode = Result;
+				break;
+			}
+		}
+		
+		if (ResultNode)
+		{
+			// If there's a return value, connect it to the result node's input pins
+			if (Form->Num() >= 2)
+			{
+				// Find first non-exec input pin on result node
+				for (UEdGraphPin* Pin : ResultNode->Pins)
+				{
+					if (Pin->Direction == EGPD_Input && 
+						Pin->PinType.PinCategory != UEdGraphSchema_K2::PC_Exec &&
+						!Pin->bHidden)
+					{
+						SetPinValueFromExpr(Pin, Form->Get(1), Ctx);
+						break;
+					}
+				}
+			}
+			
+			// Return terminates exec flow - connect to result node
+			OutExecPin = nullptr;
+			return ResultNode;
+		}
+		else
+		{
+			Ctx.Warnings.Add(TEXT("return statement used outside of function"));
+		}
+		
+		return nullptr;
+	}
+	
+	// Handle (valid? expr) - shorthand for IsValid
+	if (FormName.Equals(TEXT("valid?"), ESearchCase::IgnoreCase) && Form->Num() >= 2)
+	{
+		UFunction* IsValidFunc = UKismetSystemLibrary::StaticClass()->FindFunctionByName(TEXT("IsValid"));
+		if (IsValidFunc)
+		{
+			UK2Node_CallFunction* IsValidNode = NewObject<UK2Node_CallFunction>(Ctx.Graph);
+			IsValidNode->SetFromFunction(IsValidFunc);
+			IsValidNode->NodePosX = Ctx.CurrentX;
+			IsValidNode->NodePosY = Ctx.CurrentY;
+			Ctx.Graph->AddNode(IsValidNode, false, false);
+			IsValidNode->AllocateDefaultPins();
+			
+			// Connect the object to check
+			UEdGraphPin* ObjectPin = IsValidNode->FindPin(TEXT("Object"));
+			if (ObjectPin)
+			{
+				UEdGraphPin* ObjectSource = ResolveLispExpression(Form->Get(1), Ctx);
+				if (ObjectSource)
+				{
+					ConnectNodes(ObjectSource, ObjectPin, Ctx);
+				}
+			}
+			
+			FString TempId = Ctx.GenerateTempId();
+			Ctx.TempIdToNode.Add(TempId, IsValidNode);
+			
+			// This is a pure node, no exec pin
+			OutExecPin = nullptr;
+			return IsValidNode;
+		}
+	}
+	
+	// Handle direct function calls (FuncName args...)
+	if (!FormName.IsEmpty())
+	{
+		// Map shorthand names back to UE function names
+		static TMap<FString, FString> ShorthandToFunc;
+		if (ShorthandToFunc.Num() == 0)
+		{
+			// Math - Note: UE uses capitalized names
+			ShorthandToFunc.Add(TEXT("abs"), TEXT("Abs"));
+			ShorthandToFunc.Add(TEXT("Abs"), TEXT("Abs"));
+			ShorthandToFunc.Add(TEXT("sin"), TEXT("Sin"));
+			ShorthandToFunc.Add(TEXT("Sin"), TEXT("Sin"));
+			ShorthandToFunc.Add(TEXT("cos"), TEXT("Cos"));
+			ShorthandToFunc.Add(TEXT("Cos"), TEXT("Cos"));
+			ShorthandToFunc.Add(TEXT("tan"), TEXT("Tan"));
+			ShorthandToFunc.Add(TEXT("Tan"), TEXT("Tan"));
+			ShorthandToFunc.Add(TEXT("asin"), TEXT("Asin"));
+			ShorthandToFunc.Add(TEXT("acos"), TEXT("Acos"));
+			ShorthandToFunc.Add(TEXT("atan"), TEXT("Atan"));
+			ShorthandToFunc.Add(TEXT("atan2"), TEXT("Atan2"));
+			ShorthandToFunc.Add(TEXT("sqrt"), TEXT("Sqrt"));
+			ShorthandToFunc.Add(TEXT("Sqrt"), TEXT("Sqrt"));
+			ShorthandToFunc.Add(TEXT("sqr"), TEXT("Square"));
+			ShorthandToFunc.Add(TEXT("square"), TEXT("Square"));
+			ShorthandToFunc.Add(TEXT("pow"), TEXT("Pow"));
+			ShorthandToFunc.Add(TEXT("exp"), TEXT("Exp"));
+			ShorthandToFunc.Add(TEXT("log"), TEXT("Loge"));
+			ShorthandToFunc.Add(TEXT("log10"), TEXT("Log10"));
+			ShorthandToFunc.Add(TEXT("min"), TEXT("FMin"));
+			ShorthandToFunc.Add(TEXT("max"), TEXT("FMax"));
+			ShorthandToFunc.Add(TEXT("clamp"), TEXT("FClamp"));
+			ShorthandToFunc.Add(TEXT("Clamp"), TEXT("FClamp"));
+			ShorthandToFunc.Add(TEXT("lerp"), TEXT("Lerp"));
+			ShorthandToFunc.Add(TEXT("Lerp"), TEXT("Lerp"));
+			ShorthandToFunc.Add(TEXT("interp"), TEXT("FInterpTo"));
+			ShorthandToFunc.Add(TEXT("floor"), TEXT("Floor"));
+			ShorthandToFunc.Add(TEXT("ceil"), TEXT("Ceil"));
+			ShorthandToFunc.Add(TEXT("round"), TEXT("Round"));
+			ShorthandToFunc.Add(TEXT("trunc"), TEXT("Trunc"));
+			ShorthandToFunc.Add(TEXT("frac"), TEXT("Frac"));
+			ShorthandToFunc.Add(TEXT("sign"), TEXT("SignOfFloat"));
+			ShorthandToFunc.Add(TEXT("random"), TEXT("RandomFloat"));
+			ShorthandToFunc.Add(TEXT("random-range"), TEXT("RandomFloatInRange"));
+			ShorthandToFunc.Add(TEXT("rand-range"), TEXT("RandomFloatInRange"));
+			ShorthandToFunc.Add(TEXT("random-int"), TEXT("RandomInteger"));
+			ShorthandToFunc.Add(TEXT("random-int-range"), TEXT("RandomIntegerInRange"));
+			ShorthandToFunc.Add(TEXT("rand-int-range"), TEXT("RandomIntegerInRange"));
+			
+			// Vector
+			ShorthandToFunc.Add(TEXT("vec-length"), TEXT("VSize"));
+			ShorthandToFunc.Add(TEXT("vec-normalize"), TEXT("Normal"));
+			ShorthandToFunc.Add(TEXT("dot"), TEXT("Dot_VectorVector"));
+			ShorthandToFunc.Add(TEXT("cross"), TEXT("Cross_VectorVector"));
+			ShorthandToFunc.Add(TEXT("vec-interp"), TEXT("VInterpTo"));
+			
+			// Actor getters
+			ShorthandToFunc.Add(TEXT("get-location"), TEXT("K2_GetActorLocation"));
+			ShorthandToFunc.Add(TEXT("get-rotation"), TEXT("K2_GetActorRotation"));
+			ShorthandToFunc.Add(TEXT("get-scale"), TEXT("GetActorScale3D"));
+			ShorthandToFunc.Add(TEXT("get-velocity"), TEXT("GetVelocity"));
+			ShorthandToFunc.Add(TEXT("get-forward"), TEXT("GetActorForwardVector"));
+			ShorthandToFunc.Add(TEXT("get-right"), TEXT("GetActorRightVector"));
+			ShorthandToFunc.Add(TEXT("get-up"), TEXT("GetActorUpVector"));
+			
+			// String
+			ShorthandToFunc.Add(TEXT("str+"), TEXT("Concat_StrStr"));
+			ShorthandToFunc.Add(TEXT("str-len"), TEXT("Len"));
+			ShorthandToFunc.Add(TEXT("str-contains?"), TEXT("Contains"));
+			
+			// Conversions
+			ShorthandToFunc.Add(TEXT("int->float"), TEXT("Conv_IntToFloat"));
+			ShorthandToFunc.Add(TEXT("float->int"), TEXT("Conv_FloatToInt"));
+			ShorthandToFunc.Add(TEXT("int->string"), TEXT("Conv_IntToString"));
+			ShorthandToFunc.Add(TEXT("float->string"), TEXT("Conv_FloatToString"));
+			ShorthandToFunc.Add(TEXT("bool->string"), TEXT("Conv_BoolToString"));
+		}
+		
+		// Check if this is a shorthand
+		FString ActualFuncName = FormName;
+		if (FString* MappedName = ShorthandToFunc.Find(FormName))
+		{
+			ActualFuncName = *MappedName;
+		}
+		
+		// Build list of function name variations to try
+		TArray<FString> FuncNamesToTry;
+		FuncNamesToTry.Add(ActualFuncName);
+		FuncNamesToTry.Add(TEXT("K2_") + ActualFuncName);  // Try with K2_ prefix
+		if (ActualFuncName.StartsWith(TEXT("K2_")))
+		{
+			FuncNamesToTry.Add(ActualFuncName.Mid(3));  // Try without K2_ prefix
+		}
+		if (ActualFuncName != FormName)
+		{
+			FuncNamesToTry.Add(FormName);
+			FuncNamesToTry.Add(TEXT("K2_") + FormName);
+		}
+		
+		UFunction* Function = nullptr;
+		TArray<UClass*> ClassesToSearch = {
+			UKismetSystemLibrary::StaticClass(),
+			UKismetMathLibrary::StaticClass(),
+			UKismetArrayLibrary::StaticClass(),
+			UKismetStringLibrary::StaticClass(),
+			UGameplayStatics::StaticClass(),
+			AActor::StaticClass(),
+			APawn::StaticClass(),
+			ACharacter::StaticClass(),
+			Ctx.Blueprint->ParentClass
+		};
+		
+		// Add Geometry Script classes for procedural mesh functions
+		AddGeometryScriptClasses(ClassesToSearch);
+		
+		for (const FString& NameToTry : FuncNamesToTry)
+		{
+			for (UClass* Class : ClassesToSearch)
+			{
+				if (Class)
+				{
+					Function = Class->FindFunctionByName(*NameToTry);
+					if (Function) break;
+				}
+			}
+			if (Function) break;
+		}
+		
+		if (Function)
+		{
+			UK2Node_CallFunction* CallNode = NewObject<UK2Node_CallFunction>(Ctx.Graph);
+			CallNode->SetFromFunction(Function);
+			CallNode->NodePosX = Ctx.CurrentX;
+			CallNode->NodePosY = Ctx.CurrentY;
+			Ctx.Graph->AddNode(CallNode, false, false);
+			CallNode->AllocateDefaultPins();
+			Ctx.AdvancePosition();
+			
+			// Connect arguments
+			int32 ArgIndex = 0;
+			for (int32 i = 1; i < Form->Num(); i++)
+			{
+				if (Form->Get(i)->IsKeyword()) continue;
+				
+				int32 PinArgIndex = 0;
+				for (UEdGraphPin* Pin : CallNode->Pins)
+				{
+					if (Pin->Direction == EGPD_Input && 
+						Pin->PinType.PinCategory != UEdGraphSchema_K2::PC_Exec &&
+						Pin->PinName != UEdGraphSchema_K2::PN_Self &&
+						!Pin->bHidden)
+					{
+						if (PinArgIndex == ArgIndex)
+						{
+							UEdGraphPin* ArgSource = ResolveLispExpression(Form->Get(i), Ctx);
+							if (ArgSource)
+							{
+								ConnectNodes(ArgSource, Pin, Ctx);
+							}
+							else if (Form->Get(i)->IsNumber())
+							{
+								SetNumericPinDefaultValue(Pin, Form->Get(i)->NumberValue);
+							}
+							else if (Form->Get(i)->IsString())
+							{
+								Pin->DefaultValue = Form->Get(i)->StringValue;
+							}
+							ArgIndex++;
+							break;
+						}
+						PinArgIndex++;
+					}
+				}
+			}
+			
+			OutExecPin = GetNodeExecOutput(CallNode);
+			return CallNode;
+		}
+		else
+		{
+			Ctx.Warnings.Add(FString::Printf(TEXT("Function not found: %s"), *FormName));
+		}
+	}
+	
+	return nullptr;
+}
+
+// Convert an event form to Blueprint nodes
+static void ConvertEventForm(const FLispNodePtr& EventForm, FLispToBPContext& Ctx)
+{
+	if (!EventForm->IsList() || EventForm->Num() < 2)
+	{
+		Ctx.Errors.Add(TEXT("Invalid event form"));
+		return;
+	}
+	
+	FString EventName = EventForm->Get(1)->IsSymbol() ? EventForm->Get(1)->StringValue : TEXT("");
+
+	// Warn if user tries to use (event OnComponentBeginOverlap ...) when they should use (on-component ...)
+	// This is a common mistake - the correct syntax is (on-component ComponentName EventType body...)
+	if (EventName.Contains(TEXT("OnComponent")) || 
+		EventName.Contains(TEXT("BeginOverlap")) || 
+		EventName.Contains(TEXT("EndOverlap")) ||
+		EventName.Contains(TEXT("ComponentHit")))
+	{
+		// Extract a cleaner event type name for the suggestion
+		FString SuggestedEventType = EventName;
+		SuggestedEventType.ReplaceInline(TEXT("OnComponent"), TEXT(""));
+		SuggestedEventType.ReplaceInline(TEXT("On"), TEXT(""));
+		
+		Ctx.Warnings.Add(FString::Printf(
+			TEXT("'%s' looks like a component event. Did you mean (on-component ComponentName %s ...)? "
+			     "Example: (on-component TriggerZone BeginOverlap (set IsOpen true))"),
+			*EventName, *SuggestedEventType));
+	}
+	
+	// Map common event names
+	if (EventName.Equals(TEXT("BeginPlay"), ESearchCase::IgnoreCase))
+	{
+		EventName = TEXT("ReceiveBeginPlay");
+	}
+	else if (EventName.Equals(TEXT("Tick"), ESearchCase::IgnoreCase))
+	{
+		EventName = TEXT("ReceiveTick");
+	}
+	else if (EventName.Equals(TEXT("EndPlay"), ESearchCase::IgnoreCase))
+	{
+		EventName = TEXT("ReceiveEndPlay");
+	}
+	
+	// Check if this is an override event or custom event
+	UFunction* OverrideFunc = nullptr;
+	if (Ctx.Blueprint->ParentClass)
+	{
+		OverrideFunc = Ctx.Blueprint->ParentClass->FindFunctionByName(*EventName);
+	}
+	
+	UK2Node_Event* EventNode = nullptr;
+	
+	if (OverrideFunc)
+	{
+		EventNode = NewObject<UK2Node_Event>(Ctx.Graph);
+		EventNode->EventReference.SetFromField<UFunction>(OverrideFunc, false);
+		EventNode->bOverrideFunction = true;
+	}
+	else
+	{
+		UK2Node_CustomEvent* CustomEvent = NewObject<UK2Node_CustomEvent>(Ctx.Graph);
+		CustomEvent->CustomFunctionName = FName(*EventName);
+		EventNode = CustomEvent;
+	}
+	
+	EventNode->NodePosX = Ctx.CurrentX;
+	EventNode->NodePosY = Ctx.CurrentY;
+	Ctx.Graph->AddNode(EventNode, false, false);
+	EventNode->AllocateDefaultPins();
+	EnsureNodeHasValidGuid(EventNode);
+	Ctx.AdvancePosition();
+	
+	// Add the event node to TempIdToNode using its NodeGuid as key
+	FString EventNodeGuid = EventNode->NodeGuid.ToString();
+	Ctx.TempIdToNode.Add(EventNodeGuid, EventNode);
+	
+	// Register all output pins as variables so they can be referenced in the body
+	// This allows event parameters like DeltaTime in Tick to be used
+	for (UEdGraphPin* Pin : EventNode->Pins)
+	{
+		if (Pin->Direction == EGPD_Output && 
+			Pin->PinType.PinCategory != UEdGraphSchema_K2::PC_Exec &&
+			!Pin->bHidden)
+		{
+			FString PinName = Pin->PinName.ToString();
+			
+			// Use the event node's GUID as the key (which is now in TempIdToNode)
+			Ctx.VariableToNodeId.Add(PinName, EventNodeGuid);
+			Ctx.VariableToPin.Add(PinName, PinName);
+			
+			// Also register with common variations (no spaces, camelCase)
+			FString NoSpaces = PinName.Replace(TEXT(" "), TEXT(""));
+			if (NoSpaces != PinName)
+			{
+				Ctx.VariableToNodeId.Add(NoSpaces, EventNodeGuid);
+				Ctx.VariableToPin.Add(NoSpaces, PinName);
+			}
+			
+			// Common aliases
+			if (PinName == TEXT("Delta Seconds") || PinName == TEXT("DeltaSeconds"))
+			{
+				Ctx.VariableToNodeId.Add(TEXT("DeltaSeconds"), EventNodeGuid);
+				Ctx.VariableToPin.Add(TEXT("DeltaSeconds"), PinName);
+				Ctx.VariableToNodeId.Add(TEXT("DeltaTime"), EventNodeGuid);
+				Ctx.VariableToPin.Add(TEXT("DeltaTime"), PinName);
+			}
+		}
+	}
+	
+	// Get the execution output pin
+	UEdGraphPin* CurrentExecPin = GetNodeExecOutput(EventNode);
+	
+	// Find body statements (skip event name and :params)
+	int32 BodyStart = 2;
+	
+	// Check for :params
+	for (int32 i = 2; i < EventForm->Num(); i++)
+	{
+		if (EventForm->Get(i)->IsKeyword() && EventForm->Get(i)->StringValue == TEXT(":params"))
+		{
+			BodyStart = i + 2; // Skip :params and the params list
+			break;
+		}
+	}
+	
+	// Convert body statements
+	for (int32 i = BodyStart; i < EventForm->Num(); i++)
+	{
+		if (EventForm->Get(i)->IsKeyword()) continue;
+		
+		UEdGraphPin* NextExecPin = nullptr;
+		UEdGraphNode* Node = ConvertLispFormToNode(EventForm->Get(i), Ctx, NextExecPin);
+		
+		if (Node && CurrentExecPin)
+		{
+			UEdGraphPin* NodeInput = GetNodeExecInput(Node);
+			if (NodeInput)
+			{
+				ConnectNodes(CurrentExecPin, NodeInput, Ctx);
+			}
+		}
+		
+		if (NextExecPin)
+		{
+			CurrentExecPin = NextExecPin;
+		}
+		else if (Node)
+		{
+			// Don't try to get exec output from branch nodes - they have no single continuation
+			UK2Node_IfThenElse* BranchCheck = Cast<UK2Node_IfThenElse>(Node);
+			if (!BranchCheck)
+			{
+				CurrentExecPin = GetNodeExecOutput(Node);
+			}
+		}
+	}
+	
+	Ctx.NewRow();
+}
+
+//------------------------------------------------------------------------------
+// FECACommand_LispToBlueprint
+//------------------------------------------------------------------------------
+
+FECACommandResult FECACommand_LispToBlueprint::Execute(const TSharedPtr<FJsonObject>& Params)
+{
+	FString BlueprintPath;
+	if (!GetStringParam(Params, TEXT("blueprint_path"), BlueprintPath))
+	{
+		return FECACommandResult::Error(TEXT("Missing required parameter: blueprint_path"));
+	}
+	
+	FString Code;
+	if (!GetStringParam(Params, TEXT("code"), Code))
+	{
+		return FECACommandResult::Error(TEXT("Missing required parameter: code"));
+	}
+	
+	FString GraphName = TEXT("EventGraph");
+	GetStringParam(Params, TEXT("graph_name"), GraphName, false);
+	
+	bool bClearExisting = false;
+	Params->TryGetBoolField(TEXT("clear_existing"), bClearExisting);
+	
+	// Auto-layout option - run improved layout algorithm after conversion
+	bool bAutoLayout = false;
+	Params->TryGetBoolField(TEXT("auto_layout"), bAutoLayout);
+	
+	// Parse the Lisp code
+	FLispParseResult ParseResult = FLispParser::Parse(Code);
+	if (!ParseResult.bSuccess)
+	{
+		return FECACommandResult::Error(FString::Printf(
+			TEXT("Parse error at line %d, column %d: %s"),
+			ParseResult.ErrorLine, ParseResult.ErrorColumn, *ParseResult.Error));
+	}
+	
+	// Load blueprint
+	UBlueprint* Blueprint = LoadBlueprintByPathLisp(BlueprintPath);
+	if (!Blueprint)
+	{
+		return FECACommandResult::Error(FString::Printf(TEXT("Blueprint not found: %s"), *BlueprintPath));
+	}
+	
+	// Find graph
+	UEdGraph* Graph = FindGraphByNameLisp(Blueprint, GraphName);
+	if (!Graph)
+	{
+		return FECACommandResult::Error(FString::Printf(TEXT("Graph not found: %s"), *GraphName));
+	}
+	
+	// Clear existing nodes if requested
+	if (bClearExisting)
+	{
+		TArray<UEdGraphNode*> NodesToRemove;
+		for (UEdGraphNode* Node : Graph->Nodes)
+		{
+			NodesToRemove.Add(Node);
+		}
+		for (UEdGraphNode* Node : NodesToRemove)
+		{
+			Graph->RemoveNode(Node);
+		}
+	}
+	
+	// Set up conversion context
+	FLispToBPContext Ctx;
+	Ctx.Blueprint = Blueprint;
+	Ctx.Graph = Graph;
+	Ctx.CurrentX = 0;
+	Ctx.CurrentY = 0;
+	
+	// Process each top-level form
+	int32 EventsCreated = 0;
+	for (const FLispNodePtr& Form : ParseResult.Nodes)
+	{
+		if (!Form->IsList() || Form->Num() == 0) continue;
+		
+		FString FormName = Form->GetFormName();
+		
+		if (FormName.Equals(TEXT("event"), ESearchCase::IgnoreCase))
+		{
+			ConvertEventForm(Form, Ctx);
+			EventsCreated++;
+		}
+		else if (FormName.Equals(TEXT("on-component"), ESearchCase::IgnoreCase))
+		{
+			// Component-bound event - process through ConvertLispFormToNode
+			UEdGraphPin* OutExecPin = nullptr;
+			UEdGraphNode* Node = ConvertLispFormToNode(Form, Ctx, OutExecPin);
+			if (Node)
+			{
+				EventsCreated++;
+			}
+		}
+		else if (FormName.Equals(TEXT("on-input"), ESearchCase::IgnoreCase))
+		{
+			// Input action event - process through ConvertLispFormToNode
+			UEdGraphPin* OutExecPin = nullptr;
+			UEdGraphNode* Node = ConvertLispFormToNode(Form, Ctx, OutExecPin);
+			if (Node)
+			{
+				EventsCreated++;
+			}
+		}
+		else if (FormName.Equals(TEXT("func"), ESearchCase::IgnoreCase))
+		{
+			// Handle (func FuncName :inputs (...) :outputs (...) body...)
+			if (Form->Num() >= 2)
+			{
+				FString FuncName = Form->Get(1)->IsSymbol() ? Form->Get(1)->StringValue : TEXT("NewFunction");
+				
+				// Check if function already exists - if so, skip creation
+				bool bFunctionExists = false;
+				for (UEdGraph* ExistingGraph : Blueprint->FunctionGraphs)
+				{
+					if (ExistingGraph && ExistingGraph->GetFName() == FName(*FuncName))
+					{
+						Ctx.Warnings.Add(FString::Printf(TEXT("Function '%s' already exists, skipping creation"), *FuncName));
+						bFunctionExists = true;
+						break;
+					}
+				}
+				
+				if (bFunctionExists)
+				{
+					continue; // Skip to next form
+				}
+				
+				// Create the function graph
+				UEdGraph* FuncGraph = FBlueprintEditorUtils::CreateNewGraph(
+					Blueprint, 
+					FName(*FuncName), 
+					UEdGraph::StaticClass(), 
+					UEdGraphSchema_K2::StaticClass()
+				);
+				
+				if (FuncGraph)
+				{
+					// Add the function graph - this creates entry/result nodes automatically
+					FBlueprintEditorUtils::AddFunctionGraph<UFunction>(Blueprint, FuncGraph, /*bIsUserCreated=*/true, /*SignatureFromObject=*/nullptr);
+					
+					// Find the existing entry and result nodes created by AddFunctionGraph
+					UK2Node_FunctionEntry* EntryNode = nullptr;
+					UK2Node_FunctionResult* ResultNode = nullptr;
+					
+					for (UEdGraphNode* Node : FuncGraph->Nodes)
+					{
+						if (!EntryNode)
+						{
+							EntryNode = Cast<UK2Node_FunctionEntry>(Node);
+						}
+						if (!ResultNode)
+						{
+							ResultNode = Cast<UK2Node_FunctionResult>(Node);
+						}
+						if (EntryNode && ResultNode)
+						{
+							break;
+						}
+					}
+					
+					// Create result node if not found (some UE versions don't auto-create it)
+					if (!ResultNode)
+					{
+						ResultNode = NewObject<UK2Node_FunctionResult>(FuncGraph);
+						ResultNode->FunctionReference.SetSelfMember(FName(*FuncName));
+						ResultNode->NodePosX = 800;
+						ResultNode->NodePosY = 0;
+						FuncGraph->AddNode(ResultNode, false, false);
+						ResultNode->AllocateDefaultPins();
+					}
+					
+					if (!EntryNode)
+					{
+						Ctx.Warnings.Add(FString::Printf(TEXT("Function '%s': Could not find entry node"), *FuncName));
+						continue;
+					}
+					
+					// Parse :inputs and :outputs
+					FLispNodePtr InputsNode = Form->GetKeywordArg(TEXT(":inputs"));
+					FLispNodePtr OutputsNode = Form->GetKeywordArg(TEXT(":outputs"));
+					
+					// Helper lambda to convert Lisp type name to FEdGraphPinType
+					auto LispTypeToPinType = [](const FString& TypeName) -> FEdGraphPinType
+					{
+						FEdGraphPinType PinType;
+						PinType.PinCategory = UEdGraphSchema_K2::PC_Wildcard;
+						
+						if (TypeName.Equals(TEXT("Boolean"), ESearchCase::IgnoreCase) || TypeName.Equals(TEXT("Bool"), ESearchCase::IgnoreCase))
+						{
+							PinType.PinCategory = UEdGraphSchema_K2::PC_Boolean;
+						}
+						else if (TypeName.Equals(TEXT("Integer"), ESearchCase::IgnoreCase) || TypeName.Equals(TEXT("Int"), ESearchCase::IgnoreCase))
+						{
+							PinType.PinCategory = UEdGraphSchema_K2::PC_Int;
+						}
+						else if (TypeName.Equals(TEXT("Float"), ESearchCase::IgnoreCase) || TypeName.Equals(TEXT("Double"), ESearchCase::IgnoreCase))
+						{
+							PinType.PinCategory = UEdGraphSchema_K2::PC_Real;
+							PinType.PinSubCategory = TEXT("double");
+						}
+						else if (TypeName.Equals(TEXT("String"), ESearchCase::IgnoreCase))
+						{
+							PinType.PinCategory = UEdGraphSchema_K2::PC_String;
+						}
+						else if (TypeName.Equals(TEXT("Name"), ESearchCase::IgnoreCase))
+						{
+							PinType.PinCategory = UEdGraphSchema_K2::PC_Name;
+						}
+						else if (TypeName.Equals(TEXT("Text"), ESearchCase::IgnoreCase))
+						{
+							PinType.PinCategory = UEdGraphSchema_K2::PC_Text;
+						}
+						else if (TypeName.Equals(TEXT("Vector"), ESearchCase::IgnoreCase))
+						{
+							PinType.PinCategory = UEdGraphSchema_K2::PC_Struct;
+							PinType.PinSubCategoryObject = TBaseStructure<FVector>::Get();
+						}
+						else if (TypeName.Equals(TEXT("Rotator"), ESearchCase::IgnoreCase))
+						{
+							PinType.PinCategory = UEdGraphSchema_K2::PC_Struct;
+							PinType.PinSubCategoryObject = TBaseStructure<FRotator>::Get();
+						}
+						else if (TypeName.Equals(TEXT("Transform"), ESearchCase::IgnoreCase))
+						{
+							PinType.PinCategory = UEdGraphSchema_K2::PC_Struct;
+							PinType.PinSubCategoryObject = TBaseStructure<FTransform>::Get();
+						}
+						else if (TypeName.Equals(TEXT("Actor"), ESearchCase::IgnoreCase))
+						{
+							PinType.PinCategory = UEdGraphSchema_K2::PC_Object;
+							PinType.PinSubCategoryObject = AActor::StaticClass();
+						}
+						else if (TypeName.Equals(TEXT("Object"), ESearchCase::IgnoreCase))
+						{
+							PinType.PinCategory = UEdGraphSchema_K2::PC_Object;
+							PinType.PinSubCategoryObject = UObject::StaticClass();
+						}
+						
+						return PinType;
+					};
+					
+					// Add input parameters
+					if (InputsNode.IsValid() && InputsNode->IsList())
+					{
+						for (int32 i = 0; i < InputsNode->Num(); i++)
+						{
+							FLispNodePtr ParamDef = InputsNode->Get(i);
+							if (ParamDef->IsList() && ParamDef->Num() >= 2)
+							{
+								FString ParamName = ParamDef->Get(0)->StringValue;
+								FString ParamType = ParamDef->Get(1)->StringValue;
+								
+								FEdGraphPinType PinType = LispTypeToPinType(ParamType);
+								EntryNode->CreateUserDefinedPin(FName(*ParamName), PinType, EGPD_Output);
+							}
+						}
+					}
+					
+					// Add output parameters
+					if (OutputsNode.IsValid() && OutputsNode->IsList())
+					{
+						for (int32 i = 0; i < OutputsNode->Num(); i++)
+						{
+							FLispNodePtr ParamDef = OutputsNode->Get(i);
+							if (ParamDef->IsList() && ParamDef->Num() >= 2)
+							{
+								FString ParamName = ParamDef->Get(0)->StringValue;
+								FString ParamType = ParamDef->Get(1)->StringValue;
+								
+								FEdGraphPinType PinType = LispTypeToPinType(ParamType);
+								ResultNode->CreateUserDefinedPin(FName(*ParamName), PinType, EGPD_Input);
+							}
+						}
+					}
+					
+					// Find body start
+					int32 BodyStart = 2;
+					for (int32 i = 2; i < Form->Num(); i++)
+					{
+						if (Form->Get(i)->IsKeyword())
+						{
+							i++; // Skip keyword arg value
+							BodyStart = i + 1;
+						}
+						else
+						{
+							BodyStart = i;
+							break;
+						}
+					}
+					
+					// Create a new context for the function graph
+					FLispToBPContext FuncCtx;
+					FuncCtx.Blueprint = Blueprint;
+					FuncCtx.Graph = FuncGraph;
+					FuncCtx.CurrentX = 250;
+					FuncCtx.CurrentY = 0;
+					
+					// Register entry node for variable access
+					FString EntryTempId = FuncCtx.GenerateTempId();
+					FuncCtx.TempIdToNode.Add(EntryTempId, EntryNode);
+					
+					// Register entry pins as variables
+					for (UEdGraphPin* Pin : EntryNode->Pins)
+					{
+						if (Pin->Direction == EGPD_Output && 
+							Pin->PinType.PinCategory != UEdGraphSchema_K2::PC_Exec &&
+							!Pin->bHidden)
+						{
+							FString PinName = Pin->PinName.ToString();
+							FuncCtx.VariableToNodeId.Add(PinName, EntryNode->NodeGuid.ToString());
+							FuncCtx.VariableToPin.Add(PinName, PinName);
+						}
+					}
+					
+					// Process body statements
+					UEdGraphPin* CurrentExecPin = nullptr;
+					for (UEdGraphPin* Pin : EntryNode->Pins)
+					{
+						if (Pin->Direction == EGPD_Output && Pin->PinType.PinCategory == UEdGraphSchema_K2::PC_Exec)
+						{
+							CurrentExecPin = Pin;
+							break;
+						}
+					}
+					
+					// Ensure entry node has valid GUID
+					EnsureNodeHasValidGuid(EntryNode);
+					FuncCtx.TempIdToNode.Add(EntryNode->NodeGuid.ToString(), EntryNode);
+					
+					FuncCtx.Warnings.Add(FString::Printf(TEXT("DEBUG func '%s': Processing %d body statements starting at index %d"), *FuncName, Form->Num() - BodyStart, BodyStart));
+					
+					for (int32 i = BodyStart; i < Form->Num(); i++)
+					{
+						if (Form->Get(i)->IsKeyword()) continue;
+						
+						FString FormStr = Form->Get(i)->ToString(false);
+						FuncCtx.Warnings.Add(FString::Printf(TEXT("DEBUG func '%s': Processing statement %d: %s"), *FuncName, i, *FormStr.Left(100)));
+						
+						UEdGraphPin* NextExecPin = nullptr;
+						UEdGraphNode* Node = ConvertLispFormToNode(Form->Get(i), FuncCtx, NextExecPin);
+						
+						if (Node)
+						{
+							FuncCtx.Warnings.Add(FString::Printf(TEXT("DEBUG func '%s': Created node %s"), *FuncName, *Node->GetClass()->GetName()));
+							
+							if (CurrentExecPin)
+							{
+								UEdGraphPin* NodeInput = GetNodeExecInput(Node);
+								if (NodeInput)
+								{
+									ConnectNodes(CurrentExecPin, NodeInput, FuncCtx);
+									FuncCtx.Warnings.Add(FString::Printf(TEXT("DEBUG func '%s': Connected exec"), *FuncName));
+								}
+							}
+						}
+						else
+						{
+							FuncCtx.Warnings.Add(FString::Printf(TEXT("DEBUG func '%s': Statement %d returned no node"), *FuncName, i));
+						}
+						
+						if (NextExecPin)
+						{
+							CurrentExecPin = NextExecPin;
+							FuncCtx.Warnings.Add(FString::Printf(TEXT("DEBUG func '%s': NextExecPin was set, CurrentExecPin updated to %s"), 
+								*FuncName, *NextExecPin->PinName.ToString()));
+						}
+						else if (Node)
+						{
+							// Don't try to get exec output from branch nodes - they have no single continuation
+							UK2Node_IfThenElse* BranchCheck = Cast<UK2Node_IfThenElse>(Node);
+							FuncCtx.Warnings.Add(FString::Printf(TEXT("DEBUG func '%s': NextExecPin=null, Node=%s, BranchCheck=%s"), 
+								*FuncName, *Node->GetClass()->GetName(), BranchCheck ? TEXT("YES") : TEXT("NO")));
+							if (!BranchCheck)
+							{
+								UEdGraphPin* NodeOutput = GetNodeExecOutput(Node);
+								if (NodeOutput)
+								{
+									CurrentExecPin = NodeOutput;
+									FuncCtx.Warnings.Add(FString::Printf(TEXT("DEBUG func '%s': Set CurrentExecPin to NodeOutput %s"), 
+										*FuncName, *NodeOutput->PinName.ToString()));
+								}
+							}
+							else
+							{
+								FuncCtx.Warnings.Add(FString::Printf(TEXT("DEBUG func '%s': Branch detected, setting CurrentExecPin to nullptr"), *FuncName));
+								CurrentExecPin = nullptr;
+							}
+						}
+					}
+					
+					// Connect last exec to result node
+					FuncCtx.Warnings.Add(FString::Printf(TEXT("DEBUG func '%s': After loop, CurrentExecPin=%s"), 
+						*FuncName, CurrentExecPin ? *CurrentExecPin->PinName.ToString() : TEXT("nullptr")));
+					if (CurrentExecPin)
+					{
+						UEdGraphPin* ResultExecIn = nullptr;
+						for (UEdGraphPin* Pin : ResultNode->Pins)
+						{
+							if (Pin->Direction == EGPD_Input && Pin->PinType.PinCategory == UEdGraphSchema_K2::PC_Exec)
+							{
+								ResultExecIn = Pin;
+								break;
+							}
+						}
+						if (ResultExecIn)
+						{
+							ConnectNodes(CurrentExecPin, ResultExecIn, FuncCtx);
+						}
+					}
+					
+					// Merge warnings/errors
+					Ctx.Warnings.Append(FuncCtx.Warnings);
+					Ctx.Errors.Append(FuncCtx.Errors);
+					
+					// Mark the function graph as needing reconstruction
+					FuncGraph->NotifyGraphChanged();
+					
+					// Mark Blueprint for recompilation to register the function
+					FBlueprintEditorUtils::MarkBlueprintAsStructurallyModified(Blueprint);
+					
+					EventsCreated++; // Count functions as "events" for reporting
+					Ctx.Warnings.Add(FString::Printf(TEXT("Created function '%s' with %d nodes"), *FuncName, FuncGraph->Nodes.Num()));
+				}
+			}
+			else
+			{
+				Ctx.Warnings.Add(TEXT("Invalid func form - requires at least a name"));
+			}
+		}
+		else if (FormName.Equals(TEXT("macro"), ESearchCase::IgnoreCase))
+		{
+			Ctx.Warnings.Add(TEXT("Macro definitions not yet supported"));
+		}
+		else if (FormName.Equals(TEXT("comment"), ESearchCase::IgnoreCase))
+		{
+			// Top-level comment - create comment node
+			UEdGraphPin* OutExecPin = nullptr;
+			ConvertLispFormToNode(Form, Ctx, OutExecPin);
+		}
+	}
+	
+	// Reconstruct all nodes to resolve wildcard types after connections are made
+	// This is necessary for nodes like Cast, math operations, etc. that have
+	// wildcard pins that need to be resolved based on their connections
+	for (UEdGraphNode* Node : Graph->Nodes)
+	{
+		if (Node)
+		{
+			Node->ReconstructNode();
+		}
+	}
+	
+	// Run auto-layout if requested
+	int32 NodesLayouted = 0;
+	if (bAutoLayout)
+	{
+		FBlueprintLayoutConfig LayoutConfig;
+		// Use default spacing - increased to avoid overlaps
+		// Config already has good defaults, no need to override
+		
+		FBlueprintAutoLayout Layouter(LayoutConfig);
+		NodesLayouted = Layouter.LayoutGraph(Graph, 0, 0);
+		
+		Ctx.Warnings.Add(FString::Printf(TEXT("Auto-layout positioned %d nodes"), NodesLayouted));
+	}
+	
+	// Mark blueprint as modified
+	FBlueprintEditorUtils::MarkBlueprintAsStructurallyModified(Blueprint);
+	
+	// Build result
+	TSharedPtr<FJsonObject> Result = MakeResult();
+	Result->SetBoolField(TEXT("success"), Ctx.Errors.Num() == 0);
+	Result->SetStringField(TEXT("blueprint_path"), BlueprintPath);
+	Result->SetStringField(TEXT("graph_name"), GraphName);
+	Result->SetNumberField(TEXT("events_created"), EventsCreated);
+	Result->SetNumberField(TEXT("nodes_created"), Ctx.TempIdToNode.Num());
+	Result->SetBoolField(TEXT("auto_layout_applied"), bAutoLayout);
+	if (bAutoLayout)
+	{
+		Result->SetNumberField(TEXT("nodes_layouted"), NodesLayouted);
+	}
+	
+	if (Ctx.Errors.Num() > 0)
+	{
+		TArray<TSharedPtr<FJsonValue>> ErrorsArray;
+		for (const FString& Error : Ctx.Errors)
+		{
+			ErrorsArray.Add(MakeShared<FJsonValueString>(Error));
+		}
+		Result->SetArrayField(TEXT("errors"), ErrorsArray);
+	}
+	
+	if (Ctx.Warnings.Num() > 0)
+	{
+		TArray<TSharedPtr<FJsonValue>> WarningsArray;
+		for (const FString& Warning : Ctx.Warnings)
+		{
+			WarningsArray.Add(MakeShared<FJsonValueString>(Warning));
+		}
+		Result->SetArrayField(TEXT("warnings"), WarningsArray);
+	}
+	
+	return FECACommandResult::Success(Result);
+}
+
+//------------------------------------------------------------------------------
+// FECACommand_BlueprintLispHelp
+//------------------------------------------------------------------------------
+
+FECACommandResult FECACommand_BlueprintLispHelp::Execute(const TSharedPtr<FJsonObject>& Params)
+{
+	FString Topic;
+	GetStringParam(Params, TEXT("topic"), Topic, false);
+	Topic = Topic.ToLower();
+	
+	TSharedPtr<FJsonObject> Result = MakeResult();
+	
+	if (Topic.IsEmpty() || Topic == TEXT("all") || Topic == TEXT("forms"))
+	{
+		Result->SetStringField(TEXT("forms"), TEXT(R"(
+CORE FORMS
+==========
+
+(event EventName body...)
+  Define an event handler (BeginPlay, Tick, custom events)
+
+(func FuncName body...)
+  Define a Blueprint function
+  Example: (func DoSomething (print "Hello") (print "World"))
+
+(let VarName Expression)
+  Simple local variable binding
+  Example: (let loc (call self GetActorLocation))
+           (print (. loc X))
+
+(let ((var1 expr1) (var2 expr2)) body...)
+  Common Lisp-style let with multiple bindings and body
+  Example: (let ((x 10) (y 20)) 
+            (set sum (+ x y)))
+
+(set VarName Expression)
+  Set an existing variable's value
+
+(seq body...)
+  Sequence multiple execution statements
+
+(select condition true-value false-value)
+  Inline conditional (like ternary ?: operator)
+  Example: (select IsOpen OpenAngle 0.0)
+
+(component Name)
+  Reference a component variable by name
+  Example: (call (component Door) SetRelativeRotation ...)
+
+(make Type :field value ...)
+  Construct a struct with named fields
+  Example: (make Rotator :Pitch 0 :Yaw angle :Roll 0)
+           (make Vector :X 1 :Y 2 :Z 3)
+
+(get VarName)
+  Explicit variable get (alternative to bare symbol)
+  Example: (get MyVariable)
+)"));
+	}
+	
+	if (Topic.IsEmpty() || Topic == TEXT("all") || Topic == TEXT("flow"))
+	{
+		Result->SetStringField(TEXT("flow_control"), TEXT(R"(
+FLOW CONTROL
+============
+
+(branch Condition :true ThenBody :false ElseBody)
+  Conditional branch
+
+(seq statement1 statement2 ...)
+  Execute statements in sequence
+
+(foreach Item Collection body...)
+  Loop over a collection
+
+(delay Seconds)
+  Delay execution
+
+(cast TargetType ObjectExpr body...)
+  Cast with success execution path
+
+(switch-int Value :0 body0 :1 body1 :default defaultBody)
+  Integer switch statement
+)"));
+	}
+	
+	if (Topic.IsEmpty() || Topic == TEXT("all") || Topic == TEXT("events"))
+	{
+		Result->SetStringField(TEXT("events"), TEXT(R"(
+EVENTS
+======
+
+(event BeginPlay body...)
+  Standard Blueprint event
+
+(event Tick :params ((DeltaTime Float)) body...)
+  Event with parameters
+
+(event MyCustomEvent body...)
+  Custom event by name
+
+(on-input ActionName :pressed body :released body)
+  Input action event
+  Example: (on-input Jump :pressed (call self LaunchCharacter (vec 0 0 600)))
+
+(on-component ComponentName EventName :params (...) body...)
+  Component delegate event (BeginOverlap, EndOverlap, Hit, etc.)
+  Example: (on-component BoxCollision BeginOverlap
+            :params ((OtherActor Actor))
+            (print "Overlap!"))
+)"));
+	}
+	
+	if (Topic.IsEmpty() || Topic == TEXT("all") || Topic == TEXT("expressions"))
+	{
+		Result->SetStringField(TEXT("expressions"), TEXT(R"(
+EXPRESSIONS
+===========
+
+Function Calls:
+  (FunctionName args...)           - Global function call
+  (call Target FunctionName args...) - Call method on object
+  (call self FunctionName args...)   - Call on self
+
+Literals:
+  42, 3.14     - Numbers
+  "hello"      - Strings
+  true, false  - Booleans
+  (vec x y z)  - Vector3
+  (rot p y r)  - Rotator
+  self         - Self reference
+  nil          - Null value
+
+Struct Field Access:
+  (. struct field)  - Get field from struct (auto-creates Break node)
+  
+  Examples:
+    (let loc (GetActorLocation))
+    (print (. loc X))              ; Gets X component of vector
+    (let rot (GetActorRotation))
+    (print (. rot Yaw))            ; Gets Yaw from rotator
+
+Asset References:
+  (asset "/Game/Path/To/Asset")  - Reference any asset
+  
+  Examples:
+    (SetMaterial 0 (asset "/Game/Materials/M_MyMaterial"))
+    (PlaySound2D (asset "/Game/Sounds/S_Click"))
+    (SetStaticMesh (asset "/Game/Meshes/SM_Cube"))
+)"));
+	}
+	
+	if (Topic.IsEmpty() || Topic == TEXT("all") || Topic == TEXT("arrays"))
+	{
+		Result->SetStringField(TEXT("arrays"), TEXT(R"(
+ARRAY OPERATIONS
+================
+
+Mutating (have execution pins):
+  (array-add Array Item)       - Add item to array
+  (array-remove Array Index)   - Remove item at index
+  (array-clear Array)          - Clear all items
+  (array-shuffle Array)        - Randomize order
+  (array-reverse Array)        - Reverse order
+
+Pure (return values):
+  (array-length Array)         - Get length
+  (array-get Array Index)      - Get item at index
+  (array-contains? Array Item) - Check if contains
+  (array-find Array Item)      - Find index of item
+
+Loop:
+  (foreach Item Array body...) - Iterate over array
+)"));
+	}
+	
+	if (Topic.IsEmpty() || Topic == TEXT("all") || Topic == TEXT("functions"))
+	{
+		Result->SetStringField(TEXT("functions"), TEXT(R"(
+FUNCTIONS
+=========
+
+Define a Blueprint function:
+  (func FunctionName body...)
+
+With parameters:
+  (func FunctionName
+    :inputs ((ParamName Type) (Param2 Type2))
+    :outputs ((ReturnVal Type))
+    body...)
+
+Supported parameter types:
+  Boolean, Integer, Float, String, Name, Text, Vector, Rotator, Transform, Actor, Object
+
+Example - Simple function:
+  (func SayHello
+    (print "Hello from function!"))
+
+Example - Function with inputs:
+  (func ApplyDamage
+    :inputs ((Target Actor) (Amount Float))
+    (call Target TakeDamage Amount))
+
+Example - Function with output:
+  (func CalculateSpeed
+    :inputs ((Distance Float) (Time Float))
+    :outputs ((Speed Float))
+    (return (/ Distance Time)))
+
+Return statement:
+  (return)           ; Return with no value
+  (return value)     ; Return with value
+
+Calling functions:
+  (call self MyFunction arg1 arg2)
+)"));
+	}
+	
+	if (Topic.IsEmpty() || Topic == TEXT("all") || Topic == TEXT("math"))
+	{
+		Result->SetStringField(TEXT("math"), TEXT(R"(
+MATH OPERATORS & FUNCTIONS
+========================== Arithmetic Operators:
+  (+ a b)    - Add
+  (- a b)    - Subtract  
+  (* a b)    - Multiply
+  (/ a b)    - Divide
+  (% a b)    - Modulo Trigonometry:
+  (sin x)    - Sine
+  (cos x)    - Cosine
+  (tan x)    - Tangent
+  (asin x)   - Arc sine
+  (acos x)   - Arc cosine
+  (atan x)   - Arc tangent
+  (atan2 y x) - Two-argument arc tangent Common Math:
+  (abs x)    - Absolute value
+  (sqrt x)   - Square root
+  (pow x y)  - Power (x^y)
+  (exp x)    - e^x
+  (log x)    - Natural log
+  (floor x)  - Floor
+  (ceil x)   - Ceiling
+  (round x)  - Round
+  (sign x)   - Sign (-1, 0, or 1)
+  (frac x)   - Fractional part Clamping & Interpolation:
+  (min a b)  - Minimum
+  (max a b)  - Maximum
+  (clamp x min max) - Clamp value
+  (lerp a b t) - Linear interpolation Random:
+  (random)   - Random 0-1
+  (random-range min max) - Random float in range
+  (random-int max) - Random int 0 to max
+  (random-int-range min max) - Random int in range
+
+Comparison:
+  (< a b)    - Less than
+  (> a b)    - Greater than
+  (<= a b)   - Less or equal
+  (>= a b)   - Greater or equal
+  (= a b)    - Equal
+  (!= a b)   - Not equal
+
+Boolean:
+  (and a b)  - Logical AND
+  (or a b)   - Logical OR
+  (not a)    - Logical NOT
+
+Vector Construction:
+  (vec x y z)        - Make Vector (FVector)
+  (vec2 x y)         - Make Vector2D
+  (rot pitch yaw roll) - Make Rotator
+
+Math Functions (use as regular function calls):
+  (abs x)            - Absolute value
+  (sqrt x)           - Square root
+  (sin x) (cos x)    - Trigonometry
+  (min a b) (max a b) - Min/Max
+  (clamp x min max)  - Clamp value
+  (lerp a b t)       - Linear interpolation
+  (random)           - Random 0-1
+  (random-range a b) - Random in range
+)"));
+	}
+	
+	if (Topic.IsEmpty() || Topic == TEXT("all") || Topic == TEXT("functions"))
+	{
+		Result->SetStringField(TEXT("functions"), TEXT(R"(
+FUNCTIONS
+=========
+
+Define a Blueprint function:
+  (func FunctionName body...)
+
+Example - Simple function:
+  (func SayHello
+    (print "Hello from function!"))
+
+Example - Function with logic:
+  (func ResetPlayer
+    (set Health 100)
+    (set Ammo 30)
+    (call self SetActorLocation (vec 0 0 100)))
+
+Note: Function parameters (inputs/outputs) currently need to be 
+added manually in the Blueprint editor after creating the function.
+The function body will be fully generated.
+
+Calling functions:
+  (call self MyFunctionName)
+  (MyFunctionName)  ; If it's a global/library function
+)"));
+	}
+	
+	if (Topic.IsEmpty() || Topic == TEXT("all") || Topic == TEXT("examples"))
+	{
+		Result->SetStringField(TEXT("examples"), TEXT(R"(
+EXAMPLES
+========
+
+; ===== BASIC EVENTS =====
+
+; Simple BeginPlay
+(event BeginPlay
+  (print "Actor spawned!"))
+
+; Tick with delta time
+(event Tick
+  :params ((DeltaTime Float))
+  (let speed (* DeltaTime 100))
+  (call self AddActorWorldOffset (vec speed 0 0)))
+
+
+; ===== COLLISION & OVERLAP =====
+
+; Pickup item on overlap
+(on-component TriggerBox BeginOverlap
+  :params ((OtherActor Actor))
+  (cast Character OtherActor
+    (seq
+      (call _cast_result AddHealth 25)
+      (PlaySound2D (asset "/Game/Sounds/S_Pickup"))
+      (destroy))))
+
+; Damage on hit
+(on-component MeshComp Hit
+  :params ((OtherActor Actor))
+  (branch (valid? OtherActor)
+    :true (call OtherActor ApplyDamage 10 self)
+    :false nil))
+
+
+; ===== CONTROL FLOW =====
+
+; Branching
+(event BeginPlay
+  (branch (> Health 0)
+    :true (print "Alive!")
+    :false (print "Dead!")))
+
+; Sequence of actions
+(event BeginPlay
+  (seq
+    (set IsActive true)
+    (set Counter 0)
+    (print "Initialized!")))
+
+
+; ===== ARRAYS =====
+
+; Process all enemies
+(event BeginPlay
+  (let enemies (GetAllActorsOfClass EnemyClass))
+  (foreach enemy enemies
+    (call enemy SetTarget self)))
+
+
+; ===== FUNCTIONS =====
+
+; Define a helper function
+(func ResetStats
+  (set Health 100)
+  (set Armor 50)
+  (set Speed 600))
+
+; Call it from an event
+(event BeginPlay
+  (call self ResetStats))
+
+
+; ===== ASSETS =====
+
+; Set material on mesh
+(event BeginPlay
+  (call MeshComp SetMaterial 0 
+    (asset "/Game/Materials/M_Glowing")))
+
+; Play sound effect  
+(event BeginPlay
+  (PlaySoundAtLocation 
+    (asset "/Game/Sounds/S_Explosion")
+    (call self GetActorLocation)))
+
+
+; ===== MATH =====
+
+; Calculate damage with multiplier
+(event BeginPlay
+  (let baseDamage 10)
+  (let multiplier 1.5)
+  (let finalDamage (* baseDamage multiplier))
+  (print finalDamage))
+
+; Clamp health
+(event BeginPlay
+  (set Health (clamp Health 0 100)))
+)"));
+	}
+	
+	TArray<TSharedPtr<FJsonValue>> TopicsArray;
+	TopicsArray.Add(MakeShared<FJsonValueString>(TEXT("forms")));
+	TopicsArray.Add(MakeShared<FJsonValueString>(TEXT("events")));
+	TopicsArray.Add(MakeShared<FJsonValueString>(TEXT("flow")));
+	TopicsArray.Add(MakeShared<FJsonValueString>(TEXT("expressions")));
+	TopicsArray.Add(MakeShared<FJsonValueString>(TEXT("math")));
+	TopicsArray.Add(MakeShared<FJsonValueString>(TEXT("arrays")));
+	TopicsArray.Add(MakeShared<FJsonValueString>(TEXT("functions")));
+	TopicsArray.Add(MakeShared<FJsonValueString>(TEXT("examples")));
+	TopicsArray.Add(MakeShared<FJsonValueString>(TEXT("all")));
+	Result->SetArrayField(TEXT("available_topics"), TopicsArray);
+	
+	return FECACommandResult::Success(Result);
+}
