ECABridge experimental MCP plugin

diff --git a/Engine/Plugins/Experimental/ECABridge/Source/ECABridge/Private/Commands/ECABlueprintNodeCommands.cpp b/Engine/Plugins/Experimental/ECABridge/Source/ECABridge/Private/Commands/ECABlueprintNodeCommands.cpp
new file mode 100644
--- /dev/null
+++ b/Engine/Plugins/Experimental/ECABridge/Source/ECABridge/Private/Commands/ECABlueprintNodeCommands.cpp
@@ -0,0 +1,3881 @@
+// Copyright Epic Games, Inc. All Rights Reserved.
+
+#include "Commands/ECABlueprintNodeCommands.h"
+#include "Commands/ECACommand.h"
+#include "Editor.h"
+#include "Engine/Blueprint.h"
+#include "Engine/BlueprintGeneratedClass.h"
+#include "Kismet2/BlueprintEditorUtils.h"
+#include "Kismet2/KismetEditorUtilities.h"
+#include "EdGraph/EdGraph.h"
+#include "EdGraph/EdGraphNode.h"
+#include "EdGraph/EdGraphPin.h"
+#include "EdGraphSchema_K2.h"
+#include "K2Node_Event.h"
+#include "K2Node_InputAction.h"
+#include "K2Node_CallFunction.h"
+#include "K2Node_Self.h"
+#include "K2Node_VariableGet.h"
+#include "K2Node_VariableSet.h"
+#include "K2Node_FunctionEntry.h"
+#include "K2Node_FunctionResult.h"
+#include "K2Node_CustomEvent.h"
+#include "K2Node_ComponentBoundEvent.h"
+#include "K2Node_MacroInstance.h"
+#include "K2Node_DynamicCast.h"
+#include "K2Node_IfThenElse.h"
+#include "K2Node_Select.h"
+#include "K2Node_SwitchEnum.h"
+#include "K2Node_SwitchInteger.h"
+#include "K2Node_SwitchString.h"
+#include "K2Node_SwitchName.h"
+#include "Kismet/KismetSystemLibrary.h"
+#include "Kismet/KismetMathLibrary.h"
+#include "Kismet/KismetStringLibrary.h"
+#include "Kismet/KismetArrayLibrary.h"
+#include "Components/PrimitiveComponent.h"
+#include "Components/SceneComponent.h"
+#include "Engine/SimpleConstructionScript.h"
+#include "Engine/SCS_Node.h"
+#include "Kismet/GameplayStatics.h"
+#include "GameFramework/Actor.h"
+#include "GameFramework/Pawn.h"
+#include "GameFramework/Character.h"
+#include "Engine/Engine.h"
+#include "BlueprintAutoLayout.h"
+#include "Kismet/BlueprintFunctionLibrary.h"
+#include "UObject/UObjectIterator.h"
+
+// Register all blueprint node commands
+REGISTER_ECA_COMMAND(FECACommand_AddBlueprintEventNode)
+REGISTER_ECA_COMMAND(FECACommand_AddBlueprintInputActionNode)
+REGISTER_ECA_COMMAND(FECACommand_AddBlueprintFunctionNode)
+REGISTER_ECA_COMMAND(FECACommand_ConnectBlueprintNodes)
+REGISTER_ECA_COMMAND(FECACommand_AddBlueprintSelfReference)
+REGISTER_ECA_COMMAND(FECACommand_AddBlueprintComponentReference)
+REGISTER_ECA_COMMAND(FECACommand_FindBlueprintNodes)
+REGISTER_ECA_COMMAND(FECACommand_AddBlueprintVariableGetNode)
+REGISTER_ECA_COMMAND(FECACommand_AddBlueprintVariableSetNode)
+REGISTER_ECA_COMMAND(FECACommand_BatchEditBlueprintNodes)
+REGISTER_ECA_COMMAND(FECACommand_GetBlueprintNodePins)
+REGISTER_ECA_COMMAND(FECACommand_DeleteBlueprintNode)
+REGISTER_ECA_COMMAND(FECACommand_CleanupOrphanNodes)
+REGISTER_ECA_COMMAND(FECACommand_DisconnectBlueprintNode)
+REGISTER_ECA_COMMAND(FECACommand_SetBlueprintPinValue)
+REGISTER_ECA_COMMAND(FECACommand_SetBlueprintVariableDefault)
+REGISTER_ECA_COMMAND(FECACommand_GetBlueprintVariableDefault)
+REGISTER_ECA_COMMAND(FECACommand_BreakPinConnection)
+REGISTER_ECA_COMMAND(FECACommand_AddComponentEventNode)
+REGISTER_ECA_COMMAND(FECACommand_AddBlueprintMacroNode)
+REGISTER_ECA_COMMAND(FECACommand_AddBlueprintCastNode)
+REGISTER_ECA_COMMAND(FECACommand_DeleteBlueprintComponent)
+REGISTER_ECA_COMMAND(FECACommand_AddBlueprintFlowControlNode)
+REGISTER_ECA_COMMAND(FECACommand_AutoLayoutBlueprintGraph)
+
+//------------------------------------------------------------------------------
+// Helper Functions
+//------------------------------------------------------------------------------
+
+static UEdGraph* FindGraphByName(UBlueprint* Blueprint, const FString& GraphName)
+{
+	// Check event graphs
+	for (UEdGraph* Graph : Blueprint->UbergraphPages)
+	{
+		if (Graph->GetName() == GraphName || GraphName.Equals(TEXT("EventGraph"), ESearchCase::IgnoreCase))
+		{
+			return Graph;
+		}
+	}
+	
+	// Check function graphs
+	for (UEdGraph* Graph : Blueprint->FunctionGraphs)
+	{
+		if (Graph->GetName() == GraphName)
+		{
+			return Graph;
+		}
+	}
+	
+	return nullptr;
+}
+
+// Constants for auto-layout spacing
+static const int32 AUTO_LAYOUT_SPACING_X = 400;
+static const int32 AUTO_LAYOUT_SPACING_Y = 150;
+
+// Calculate automatic position for a new node in a Blueprint graph
+// Places the node to the right of the rightmost node, or below if there's a recent node
+static FVector2D CalculateAutoNodePosition(UEdGraph* Graph, UEdGraphNode* ConnectToNode = nullptr)
+{
+	if (!Graph)
+	{
+		return FVector2D(0, 0);
+	}
+	
+	// If we have a node to connect to, position relative to it
+	if (ConnectToNode)
+	{
+		// Place to the right of the node we're connecting from
+		return FVector2D(
+			ConnectToNode->NodePosX + AUTO_LAYOUT_SPACING_X,
+			ConnectToNode->NodePosY
+		);
+	}
+	
+	// Find the bounds of existing nodes
+	int32 MaxX = INT_MIN;
+	int32 MaxY = INT_MIN;
+	int32 MinX = INT_MAX;
+	int32 MinY = INT_MAX;
+	int32 NodeCount = 0;
+	
+	// Track the rightmost node at each Y level to find a good insertion point
+	TMap<int32, int32> RightmostXAtY; // Y bucket -> max X
+	
+	for (UEdGraphNode* Node : Graph->Nodes)
+	{
+		if (Node)
+		{
+			MaxX = FMath::Max(MaxX, Node->NodePosX);
+			MaxY = FMath::Max(MaxY, Node->NodePosY);
+			MinX = FMath::Min(MinX, Node->NodePosX);
+			MinY = FMath::Min(MinY, Node->NodePosY);
+			NodeCount++;
+			
+			// Track rightmost node in Y buckets (rounded to spacing)
+			int32 YBucket = (Node->NodePosY / AUTO_LAYOUT_SPACING_Y) * AUTO_LAYOUT_SPACING_Y;
+			if (!RightmostXAtY.Contains(YBucket) || Node->NodePosX > RightmostXAtY[YBucket])
+			{
+				RightmostXAtY.Add(YBucket, Node->NodePosX);
+			}
+		}
+	}
+	
+	// If no nodes exist, start at origin
+	if (NodeCount == 0)
+	{
+		return FVector2D(0, 0);
+	}
+	
+	// Find a good position: to the right of existing nodes
+	// Try to find a Y level that isn't too crowded
+	int32 BestY = MinY;
+	int32 BestX = MaxX + AUTO_LAYOUT_SPACING_X;
+	
+	// If the graph is getting wide, start a new row
+	if (MaxX - MinX > AUTO_LAYOUT_SPACING_X * 5)
+	{
+		// Find the Y level with the least rightward extent
+		int32 MinRightX = INT_MAX;
+		for (const auto& Pair : RightmostXAtY)
+		{
+			if (Pair.Value < MinRightX)
+			{
+				MinRightX = Pair.Value;
+				BestY = Pair.Key;
+				BestX = MinRightX + AUTO_LAYOUT_SPACING_X;
+			}
+		}
+		
+		// Or start a completely new row below everything
+		if (BestX > MaxX)
+		{
+			BestY = MaxY + AUTO_LAYOUT_SPACING_Y;
+			BestX = MinX;
+		}
+	}
+	
+	return FVector2D(BestX, BestY);
+}
+
+static FVector2D GetNodePosition(const TSharedPtr<FJsonObject>& Params, UEdGraph* Graph = nullptr, UEdGraphNode* ConnectToNode = nullptr)
+{
+	FVector2D Position(0, 0);
+	const TSharedPtr<FJsonObject>* PosObj;
+	
+	// Check if explicit position was provided
+	if (Params->TryGetObjectField(TEXT("node_position"), PosObj))
+	{
+		Position.X = (*PosObj)->GetNumberField(TEXT("x"));
+		Position.Y = (*PosObj)->GetNumberField(TEXT("y"));
+	}
+	else if (Graph)
+	{
+		// No position provided - calculate automatically
+		Position = CalculateAutoNodePosition(Graph, ConnectToNode);
+	}
+	
+	return Position;
+}
+
+// Helper to ensure a node has a valid GUID
+// Some nodes may not have their GUID properly initialized when created programmatically
+static void EnsureNodeHasValidGuid(UEdGraphNode* Node)
+{
+	if (Node && !Node->NodeGuid.IsValid())
+	{
+		Node->CreateNewGuid();
+	}
+}
+
+static TSharedPtr<FJsonObject> PinToJson(UEdGraphPin* Pin)
+{
+	TSharedPtr<FJsonObject> PinJson = MakeShared<FJsonObject>();
+	
+	// Get the actual pin name for connections
+	FString ActualPinName = Pin->PinName.ToString();
+	FString DisplayName = ActualPinName;
+	
+	// Handle confusing "None" pin names - provide better display names for exec pins
+	if (ActualPinName.IsEmpty() || ActualPinName == TEXT("None"))
+	{
+		if (Pin->PinType.PinCategory == UEdGraphSchema_K2::PC_Exec)
+		{
+			// Exec pins with no name - give them a meaningful display name
+			DisplayName = Pin->Direction == EGPD_Input ? TEXT("execute") : TEXT("then");
+		}
+		else
+		{
+			DisplayName = Pin->Direction == EGPD_Input ? TEXT("input") : TEXT("output");
+		}
+	}
+	
+	// "name" is the user-friendly name to display and use in connections
+	// The connection code will handle mapping these back to actual pin names
+	PinJson->SetStringField(TEXT("name"), DisplayName);
+	
+	// Also include the actual internal pin name if different (for debugging)
+	if (DisplayName != ActualPinName)
+	{
+		PinJson->SetStringField(TEXT("internal_name"), ActualPinName);
+	}
+	PinJson->SetStringField(TEXT("direction"), Pin->Direction == EGPD_Input ? TEXT("Input") : TEXT("Output"));
+	PinJson->SetStringField(TEXT("type"), Pin->PinType.PinCategory.ToString());
+	PinJson->SetStringField(TEXT("sub_type"), Pin->PinType.PinSubCategory.ToString());
+	
+	// Include sub-category object for struct/class types
+	if (Pin->PinType.PinSubCategoryObject.IsValid())
+	{
+		PinJson->SetStringField(TEXT("sub_type_object"), Pin->PinType.PinSubCategoryObject->GetName());
+	}
+	
+	PinJson->SetBoolField(TEXT("is_connected"), Pin->LinkedTo.Num() > 0);
+	PinJson->SetBoolField(TEXT("is_hidden"), Pin->bHidden);
+	
+	// Include default value if set
+	if (!Pin->DefaultValue.IsEmpty())
+	{
+		PinJson->SetStringField(TEXT("default_value"), Pin->DefaultValue);
+	}
+	
+	// Include default object reference for object-type pins (materials, textures, etc.)
+	if (Pin->DefaultObject)
+	{
+		PinJson->SetStringField(TEXT("default_object"), Pin->DefaultObject->GetPathName());
+		PinJson->SetStringField(TEXT("default_object_name"), Pin->DefaultObject->GetName());
+		PinJson->SetStringField(TEXT("default_object_class"), Pin->DefaultObject->GetClass()->GetName());
+	}
+	
+	// Include AutogeneratedDefaultValue if different from DefaultValue
+	if (!Pin->AutogeneratedDefaultValue.IsEmpty() && Pin->AutogeneratedDefaultValue != Pin->DefaultValue)
+	{
+		PinJson->SetStringField(TEXT("autogenerated_default"), Pin->AutogeneratedDefaultValue);
+	}
+	
+	// Include connected pins info
+	if (Pin->LinkedTo.Num() > 0)
+	{
+		TArray<TSharedPtr<FJsonValue>> ConnectedArray;
+		for (UEdGraphPin* LinkedPin : Pin->LinkedTo)
+		{
+			TSharedPtr<FJsonObject> LinkedInfo = MakeShared<FJsonObject>();
+			LinkedInfo->SetStringField(TEXT("node_id"), LinkedPin->GetOwningNode()->NodeGuid.ToString());
+			LinkedInfo->SetStringField(TEXT("pin_name"), LinkedPin->PinName.ToString());
+			ConnectedArray.Add(MakeShared<FJsonValueObject>(LinkedInfo));
+		}
+		PinJson->SetArrayField(TEXT("connected_to"), ConnectedArray);
+	}
+	
+	return PinJson;
+}
+
+static TSharedPtr<FJsonObject> NodeToJson(UEdGraphNode* Node)
+{
+	// Ensure the node has a valid GUID before serializing
+	EnsureNodeHasValidGuid(Node);
+	
+	TSharedPtr<FJsonObject> NodeJson = MakeShared<FJsonObject>();
+	NodeJson->SetStringField(TEXT("node_id"), Node->NodeGuid.ToString());
+	NodeJson->SetStringField(TEXT("node_class"), Node->GetClass()->GetName());
+	NodeJson->SetStringField(TEXT("node_title"), Node->GetNodeTitle(ENodeTitleType::FullTitle).ToString());
+	NodeJson->SetNumberField(TEXT("x"), Node->NodePosX);
+	NodeJson->SetNumberField(TEXT("y"), Node->NodePosY);
+	
+	// Include error/warning info if present
+	NodeJson->SetBoolField(TEXT("has_error"), Node->bHasCompilerMessage);
+	if (Node->bHasCompilerMessage && !Node->ErrorMsg.IsEmpty())
+	{
+		NodeJson->SetStringField(TEXT("error_message"), Node->ErrorMsg);
+	}
+	
+	// Include pins with detailed info
+	TArray<TSharedPtr<FJsonValue>> PinsArray;
+	for (UEdGraphPin* Pin : Node->Pins)
+	{
+		if (!Pin->bHidden)
+		{
+			PinsArray.Add(MakeShared<FJsonValueObject>(PinToJson(Pin)));
+		}
+	}
+	NodeJson->SetArrayField(TEXT("pins"), PinsArray);
+	
+	return NodeJson;
+}
+
+// Helper to add error info from a node to a result object
+static void AddNodeErrorInfo(TSharedPtr<FJsonObject>& Result, UEdGraphNode* Node)
+{
+	if (Node && Node->bHasCompilerMessage)
+	{
+		Result->SetBoolField(TEXT("has_error"), true);
+		if (!Node->ErrorMsg.IsEmpty())
+		{
+			Result->SetStringField(TEXT("error_message"), Node->ErrorMsg);
+		}
+	}
+}
+
+// Helper to find a node by GUID in a graph
+static UEdGraphNode* FindNodeByGuid(UEdGraph* Graph, const FGuid& Guid)
+{
+	for (UEdGraphNode* Node : Graph->Nodes)
+	{
+		if (Node->NodeGuid == Guid)
+		{
+			return Node;
+		}
+	}
+	return nullptr;
+}
+
+// Helper to resolve a node reference (either temp_id or GUID)
+static UEdGraphNode* ResolveNodeReference(UEdGraph* Graph, const FString& Reference, const TMap<FString, UEdGraphNode*>& TempIdMap)
+{
+	// First check if it's a temp_id
+	if (const UEdGraphNode* const* Found = TempIdMap.Find(Reference))
+	{
+		return const_cast<UEdGraphNode*>(*Found);
+	}
+	
+	// Otherwise try to parse as GUID
+	FGuid Guid;
+	if (FGuid::Parse(Reference, Guid))
+	{
+		return FindNodeByGuid(Graph, Guid);
+	}
+	
+	return nullptr;
+}
+
+// Helper to find a pin by name, handling friendly name aliases
+// Maps "execute" -> "" or "None" for exec input pins, "then" -> "" or "None" for exec output pins
+static UEdGraphPin* FindPinByFriendlyName(UEdGraphNode* Node, const FString& PinName, EEdGraphPinDirection Direction = EGPD_MAX)
+{
+	if (!Node)
+	{
+		return nullptr;
+	}
+	
+	// First try exact match with direction
+	if (Direction != EGPD_MAX)
+	{
+		if (UEdGraphPin* Pin = Node->FindPin(*PinName, Direction))
+		{
+			return Pin;
+		}
+	}
+	
+	// Try exact match without direction
+	if (UEdGraphPin* Pin = Node->FindPin(*PinName))
+	{
+		return Pin;
+	}
+	
+	// Handle friendly name aliases for exec pins with "None" as their actual name
+	FString LowerPinName = PinName.ToLower();
+	if (LowerPinName == TEXT("execute") || LowerPinName == TEXT("exec") || LowerPinName == TEXT("in"))
+	{
+		// Look for input exec pin with empty or "None" name
+		for (UEdGraphPin* Pin : Node->Pins)
+		{
+			if (Pin->Direction == EGPD_Input && 
+				Pin->PinType.PinCategory == UEdGraphSchema_K2::PC_Exec &&
+				(Pin->PinName.IsNone() || Pin->PinName.ToString().IsEmpty() || Pin->PinName.ToString() == TEXT("None")))
+			{
+				return Pin;
+			}
+		}
+	}
+	else if (LowerPinName == TEXT("then") || LowerPinName == TEXT("out"))
+	{
+		// Look for output exec pin with empty or "None" name
+		for (UEdGraphPin* Pin : Node->Pins)
+		{
+			if (Pin->Direction == EGPD_Output && 
+				Pin->PinType.PinCategory == UEdGraphSchema_K2::PC_Exec &&
+				(Pin->PinName.IsNone() || Pin->PinName.ToString().IsEmpty() || Pin->PinName.ToString() == TEXT("None")))
+			{
+				return Pin;
+			}
+		}
+	}
+	
+	// Also try "None" directly in case user specifies it
+	if (LowerPinName == TEXT("none"))
+	{
+		if (UEdGraphPin* Pin = Node->FindPin(NAME_None, Direction != EGPD_MAX ? Direction : EGPD_Input))
+		{
+			return Pin;
+		}
+		if (UEdGraphPin* Pin = Node->FindPin(NAME_None))
+		{
+			return Pin;
+		}
+	}
+	
+	return nullptr;
+}
+
+//------------------------------------------------------------------------------
+// AddBlueprintEventNode
+//------------------------------------------------------------------------------
+
+FECACommandResult FECACommand_AddBlueprintEventNode::Execute(const TSharedPtr<FJsonObject>& Params)
+{
+	FString BlueprintPath;
+	if (!GetStringParam(Params, TEXT("blueprint_path"), BlueprintPath))
+	{
+		return FECACommandResult::Error(TEXT("Missing required parameter: blueprint_path"));
+	}
+	
+	FString EventName;
+	if (!GetStringParam(Params, TEXT("event_name"), EventName))
+	{
+		return FECACommandResult::Error(TEXT("Missing required parameter: event_name"));
+	}
+	
+	UBlueprint* Blueprint = LoadBlueprintByPath(BlueprintPath);
+	if (!Blueprint)
+	{
+		return FECACommandResult::Error(FString::Printf(TEXT("Blueprint not found: %s"), *BlueprintPath));
+	}
+	
+	// Get the event graph
+	UEdGraph* EventGraph = FindGraphByName(Blueprint, TEXT("EventGraph"));
+	if (!EventGraph)
+	{
+		return FECACommandResult::Error(TEXT("Event graph not found"));
+	}
+	
+	FVector2D Position = GetNodePosition(Params, EventGraph);
+	
+	// Find the function to create an event for
+	UFunction* OverrideFunc = nullptr;
+	UClass* ParentClass = Blueprint->ParentClass;
+	
+	if (ParentClass)
+	{
+		OverrideFunc = ParentClass->FindFunctionByName(*EventName);
+	}
+	
+	UK2Node_Event* EventNode = nullptr;
+	
+	if (OverrideFunc)
+	{
+		// Create event node for the override function
+		EventNode = NewObject<UK2Node_Event>(EventGraph);
+		EventNode->EventReference.SetFromField<UFunction>(OverrideFunc, false);
+		EventNode->bOverrideFunction = true;
+	}
+	else
+	{
+		// Create a custom event
+		UK2Node_CustomEvent* CustomEvent = NewObject<UK2Node_CustomEvent>(EventGraph);
+		CustomEvent->CustomFunctionName = FName(*EventName);
+		EventNode = CustomEvent;
+	}
+	
+	if (!EventNode)
+	{
+		return FECACommandResult::Error(TEXT("Failed to create event node"));
+	}
+	
+	EventNode->NodePosX = Position.X;
+	EventNode->NodePosY = Position.Y;
+	EventGraph->AddNode(EventNode, false, false);
+	EventNode->AllocateDefaultPins();
+	EnsureNodeHasValidGuid(EventNode);
+	EventNode->ReconstructNode();
+	
+	FBlueprintEditorUtils::MarkBlueprintAsStructurallyModified(Blueprint);
+	
+	TSharedPtr<FJsonObject> Result = MakeResult();
+	Result->SetStringField(TEXT("node_id"), EventNode->NodeGuid.ToString());
+	Result->SetStringField(TEXT("event_name"), EventName);
+	Result->SetObjectField(TEXT("node"), NodeToJson(EventNode));
+	AddNodeErrorInfo(Result, EventNode);
+	
+	return FECACommandResult::Success(Result);
+}
+
+//------------------------------------------------------------------------------
+// AddBlueprintInputActionNode
+//------------------------------------------------------------------------------
+
+FECACommandResult FECACommand_AddBlueprintInputActionNode::Execute(const TSharedPtr<FJsonObject>& Params)
+{
+	FString BlueprintPath;
+	if (!GetStringParam(Params, TEXT("blueprint_path"), BlueprintPath))
+	{
+		return FECACommandResult::Error(TEXT("Missing required parameter: blueprint_path"));
+	}
+	
+	FString ActionName;
+	if (!GetStringParam(Params, TEXT("action_name"), ActionName))
+	{
+		return FECACommandResult::Error(TEXT("Missing required parameter: action_name"));
+	}
+	
+	UBlueprint* Blueprint = LoadBlueprintByPath(BlueprintPath);
+	if (!Blueprint)
+	{
+		return FECACommandResult::Error(FString::Printf(TEXT("Blueprint not found: %s"), *BlueprintPath));
+	}
+	
+	UEdGraph* EventGraph = FindGraphByName(Blueprint, TEXT("EventGraph"));
+	if (!EventGraph)
+	{
+		return FECACommandResult::Error(TEXT("Event graph not found"));
+	}
+	
+	FVector2D Position = GetNodePosition(Params, EventGraph);
+	
+	// Create input action node
+	UK2Node_InputAction* InputActionNode = NewObject<UK2Node_InputAction>(EventGraph);
+	InputActionNode->InputActionName = FName(*ActionName);
+	InputActionNode->NodePosX = Position.X;
+	InputActionNode->NodePosY = Position.Y;
+	
+	EventGraph->AddNode(InputActionNode, false, false);
+	InputActionNode->AllocateDefaultPins();
+	EnsureNodeHasValidGuid(InputActionNode);
+	
+	FBlueprintEditorUtils::MarkBlueprintAsStructurallyModified(Blueprint);
+	
+	TSharedPtr<FJsonObject> Result = MakeResult();
+	Result->SetStringField(TEXT("node_id"), InputActionNode->NodeGuid.ToString());
+	Result->SetStringField(TEXT("action_name"), ActionName);
+	Result->SetObjectField(TEXT("node"), NodeToJson(InputActionNode));
+	
+	return FECACommandResult::Success(Result);
+}
+
+//------------------------------------------------------------------------------
+// AddBlueprintFunctionNode
+//------------------------------------------------------------------------------
+
+FECACommandResult FECACommand_AddBlueprintFunctionNode::Execute(const TSharedPtr<FJsonObject>& Params)
+{
+	FString BlueprintPath;
+	if (!GetStringParam(Params, TEXT("blueprint_path"), BlueprintPath))
+	{
+		return FECACommandResult::Error(TEXT("Missing required parameter: blueprint_path"));
+	}
+	
+	FString FunctionName;
+	if (!GetStringParam(Params, TEXT("function_name"), FunctionName))
+	{
+		return FECACommandResult::Error(TEXT("Missing required parameter: function_name"));
+	}
+	
+	FString Target = TEXT("self");
+	GetStringParam(Params, TEXT("target"), Target, false);
+	
+	FString TargetClass;
+	GetStringParam(Params, TEXT("target_class"), TargetClass, false);
+	
+	FString GraphName = TEXT("EventGraph");
+	GetStringParam(Params, TEXT("graph_name"), GraphName, false);
+	
+	UBlueprint* Blueprint = LoadBlueprintByPath(BlueprintPath);
+	if (!Blueprint)
+	{
+		return FECACommandResult::Error(FString::Printf(TEXT("Blueprint not found: %s"), *BlueprintPath));
+	}
+	
+	UEdGraph* Graph = FindGraphByName(Blueprint, GraphName);
+	if (!Graph)
+	{
+		return FECACommandResult::Error(FString::Printf(TEXT("Graph not found: %s"), *GraphName));
+	}
+	
+	FVector2D Position = GetNodePosition(Params, Graph);
+	
+	// Find the function
+	UFunction* Function = nullptr;
+	
+	// Try common library classes first
+	TArray<UClass*> ClassesToSearch;
+	
+	if (!TargetClass.IsEmpty())
+	{
+		UClass* SpecificClass = FindObject<UClass>(nullptr, *FString::Printf(TEXT("/Script/Engine.%s"), *TargetClass));
+		if (!SpecificClass)
+		{
+			SpecificClass = FindObject<UClass>(nullptr, *FString::Printf(TEXT("/Script/CoreUObject.%s"), *TargetClass));
+		}
+		if (SpecificClass)
+		{
+			ClassesToSearch.Add(SpecificClass);
+		}
+	}
+	
+	// First, check if this is a function defined in the Blueprint itself
+	bool bIsBlueprintFunction = false;
+	for (UEdGraph* FuncGraph : Blueprint->FunctionGraphs)
+	{
+		if (FuncGraph && FuncGraph->GetFName() == FName(*FunctionName))
+		{
+			bIsBlueprintFunction = true;
+			break;
+		}
+	}
+	
+	if (bIsBlueprintFunction)
+	{
+		// Create a call to a Blueprint-defined function
+		UK2Node_CallFunction* CallNode = NewObject<UK2Node_CallFunction>(Graph);
+		CallNode->FunctionReference.SetSelfMember(FName(*FunctionName));
+		CallNode->NodePosX = Position.X;
+		CallNode->NodePosY = Position.Y;
+		Graph->AddNode(CallNode, false, false);
+		CallNode->AllocateDefaultPins();
+		EnsureNodeHasValidGuid(CallNode);
+		CallNode->ReconstructNode();
+		
+		FBlueprintEditorUtils::MarkBlueprintAsStructurallyModified(Blueprint);
+		
+		TSharedPtr<FJsonObject> Result = MakeResult();
+		Result->SetStringField(TEXT("node_id"), CallNode->NodeGuid.ToString());
+		Result->SetStringField(TEXT("function_name"), FunctionName);
+		Result->SetStringField(TEXT("function_class"), TEXT("Self"));
+		Result->SetBoolField(TEXT("is_blueprint_function"), true);
+		Result->SetObjectField(TEXT("node"), NodeToJson(CallNode));
+		AddNodeErrorInfo(Result, CallNode);
+		
+		return FECACommandResult::Success(Result);
+	}
+	
+	// Add common classes for engine/library functions
+	ClassesToSearch.Add(UKismetSystemLibrary::StaticClass());
+	ClassesToSearch.Add(UKismetMathLibrary::StaticClass());
+	ClassesToSearch.Add(UKismetStringLibrary::StaticClass());
+	ClassesToSearch.Add(UKismetArrayLibrary::StaticClass());
+	ClassesToSearch.Add(UGameplayStatics::StaticClass());
+	ClassesToSearch.Add(AActor::StaticClass());
+	ClassesToSearch.Add(APawn::StaticClass());
+	ClassesToSearch.Add(ACharacter::StaticClass());
+	ClassesToSearch.Add(USceneComponent::StaticClass());
+	ClassesToSearch.Add(UPrimitiveComponent::StaticClass());
+	ClassesToSearch.Add(Blueprint->ParentClass);
+	
+	// Add Geometry Script library classes for procedural mesh operations
+	// These provide AppendBox, AppendSphere, AppendCylinder, and other mesh primitives
+	static const TCHAR* GeometryScriptClassPaths[] = {
+		// Mesh primitive creation (AppendBox, AppendSphere, AppendCylinder, etc.)
+		TEXT("/Script/GeometryScriptingCore.GeometryScriptLibrary_MeshPrimitiveFunctions"),
+		// Basic mesh editing (AppendMesh, SetVertexPosition, etc.)
+		TEXT("/Script/GeometryScriptingCore.GeometryScriptLibrary_MeshBasicEditFunctions"),
+		// Transform operations (TransformMesh, TranslateVertices, etc.)
+		TEXT("/Script/GeometryScriptingCore.GeometryScriptLibrary_MeshTransformFunctions"),
+		// Boolean operations (MeshBooleanUnion, MeshBooleanIntersect, etc.)
+		TEXT("/Script/GeometryScriptingCore.GeometryScriptLibrary_MeshBooleanFunctions"),
+		// Modeling operations (extrude, offset, etc.)
+		TEXT("/Script/GeometryScriptingCore.GeometryScriptLibrary_MeshModelingFunctions"),
+		// Normal operations (ComputeNormals, FlipNormals, etc.)
+		TEXT("/Script/GeometryScriptingCore.GeometryScriptLibrary_MeshNormalsFunctions"),
+		// UV operations (SetMeshUVs, AutoGenerateUVs, etc.)
+		TEXT("/Script/GeometryScriptingCore.GeometryScriptLibrary_MeshUVFunctions"),
+		// Material operations (SetMaterialID, etc.)
+		TEXT("/Script/GeometryScriptingCore.GeometryScriptLibrary_MeshMaterialFunctions"),
+		// Deformation operations (Bend, Twist, etc.)
+		TEXT("/Script/GeometryScriptingCore.GeometryScriptLibrary_MeshDeformFunctions"),
+		// Subdivision operations
+		TEXT("/Script/GeometryScriptingCore.GeometryScriptLibrary_MeshSubdivideFunctions"),
+		// Simplification operations (Decimate, etc.)
+		TEXT("/Script/GeometryScriptingCore.GeometryScriptLibrary_MeshSimplifyFunctions"),
+		// Repair operations (FillHoles, WeldVertices, etc.)
+		TEXT("/Script/GeometryScriptingCore.GeometryScriptLibrary_MeshRepairFunctions"),
+		// Query operations (GetVertexCount, GetBoundingBox, etc.)
+		TEXT("/Script/GeometryScriptingCore.GeometryScriptLibrary_MeshQueryFunctions"),
+		// Spatial operations (FindNearestPoint, etc.)
+		TEXT("/Script/GeometryScriptingCore.GeometryScriptLibrary_MeshSpatialFunctions"),
+		// Selection operations
+		TEXT("/Script/GeometryScriptingCore.GeometryScriptLibrary_MeshSelectionFunctions"),
+		// Containment checks
+		TEXT("/Script/GeometryScriptingCore.GeometryScriptLibrary_ContainmentFunctions"),
+		// Collision functions
+		TEXT("/Script/GeometryScriptingCore.GeometryScriptLibrary_CollisionFunctions"),
+		// Voxel operations
+		TEXT("/Script/GeometryScriptingCore.GeometryScriptLibrary_MeshVoxelFunctions"),
+		// Path/polyline operations
+		TEXT("/Script/GeometryScriptingCore.GeometryScriptLibrary_PolyPathFunctions"),
+		// List utilities
+		TEXT("/Script/GeometryScriptingCore.GeometryScriptLibrary_ListUtilityFunctions"),
+		// Asset creation (editor only)
+		TEXT("/Script/GeometryScriptingEditor.GeometryScriptLibrary_CreateNewAssetUtilityFunctions"),
+	};
+	for (const TCHAR* ClassPath : GeometryScriptClassPaths)
+	{
+		// Try FindObject first (works if class is already loaded)
+		UClass* GeoClass = FindObject<UClass>(nullptr, ClassPath);
+		if (!GeoClass)
+		{
+			// Try LoadObject (loads the class if not yet loaded)
+			GeoClass = LoadObject<UClass>(nullptr, ClassPath);
+		}
+		if (!GeoClass)
+		{
+			// Try StaticLoadClass as a fallback
+			GeoClass = StaticLoadClass(UObject::StaticClass(), nullptr, ClassPath);
+		}
+		if (GeoClass)
+		{
+			ClassesToSearch.Add(GeoClass);
+		}
+	}
+	if (Blueprint->GeneratedClass)
+	{
+		ClassesToSearch.Add(Blueprint->GeneratedClass);
+	}
+	
+	for (UClass* Class : ClassesToSearch)
+	{
+		if (Class)
+		{
+			Function = Class->FindFunctionByName(*FunctionName);
+			if (Function)
+			{
+				break;
+			}
+			// Also try with K2_ prefix
+			Function = Class->FindFunctionByName(*FString::Printf(TEXT("K2_%s"), *FunctionName));
+			if (Function)
+			{
+				break;
+			}
+		}
+	}
+	
+	// If still not found, do a broader search across all BlueprintFunctionLibrary classes
+	if (!Function)
+	{
+		for (TObjectIterator<UClass> It; It; ++It)
+		{
+			UClass* TestClass = *It;
+			if (TestClass->IsChildOf(UBlueprintFunctionLibrary::StaticClass()) && !TestClass->HasAnyClassFlags(CLASS_Abstract))
+			{
+				Function = TestClass->FindFunctionByName(*FunctionName);
+				if (Function)
+				{
+					UE_LOG(LogTemp, Log, TEXT("[ECABridge] Found function '%s' in class '%s'"), *FunctionName, *TestClass->GetName());
+					break;
+				}
+				// Try with K2_ prefix
+				Function = TestClass->FindFunctionByName(*FString::Printf(TEXT("K2_%s"), *FunctionName));
+				if (Function)
+				{
+					UE_LOG(LogTemp, Log, TEXT("[ECABridge] Found function 'K2_%s' in class '%s'"), *FunctionName, *TestClass->GetName());
+					break;
+				}
+			}
+		}
+	}
+	
+	if (!Function)
+	{
+		return FECACommandResult::Error(FString::Printf(TEXT("Function not found: %s. Searched %d classes including Geometry Script libraries."), *FunctionName, ClassesToSearch.Num()));
+	}
+	
+	// Create the call function node
+	UK2Node_CallFunction* CallNode = NewObject<UK2Node_CallFunction>(Graph);
+	CallNode->SetFromFunction(Function);
+	CallNode->NodePosX = Position.X;
+	CallNode->NodePosY = Position.Y;
+	
+	Graph->AddNode(CallNode, false, false);
+	CallNode->AllocateDefaultPins();
+	EnsureNodeHasValidGuid(CallNode);
+	
+	// Reconstruct node to ensure proper setup and error detection
+	CallNode->ReconstructNode();
+	
+	FBlueprintEditorUtils::MarkBlueprintAsStructurallyModified(Blueprint);
+	
+	TSharedPtr<FJsonObject> Result = MakeResult();
+	Result->SetStringField(TEXT("node_id"), CallNode->NodeGuid.ToString());
+	Result->SetStringField(TEXT("function_name"), FunctionName);
+	Result->SetStringField(TEXT("function_class"), Function->GetOwnerClass()->GetName());
+	Result->SetObjectField(TEXT("node"), NodeToJson(CallNode));
+	
+	// Include any errors at the top level for easy access
+	AddNodeErrorInfo(Result, CallNode);
+	
+	return FECACommandResult::Success(Result);
+}
+
+//------------------------------------------------------------------------------
+// ConnectBlueprintNodes
+//------------------------------------------------------------------------------
+
+FECACommandResult FECACommand_ConnectBlueprintNodes::Execute(const TSharedPtr<FJsonObject>& Params)
+{
+	FString BlueprintPath;
+	if (!GetStringParam(Params, TEXT("blueprint_path"), BlueprintPath))
+	{
+		return FECACommandResult::Error(TEXT("Missing required parameter: blueprint_path"));
+	}
+	
+	FString SourceNodeId, TargetNodeId;
+	FString SourcePinName, TargetPinName;
+	
+	if (!GetStringParam(Params, TEXT("source_node_id"), SourceNodeId))
+	{
+		return FECACommandResult::Error(TEXT("Missing required parameter: source_node_id"));
+	}
+	if (!GetStringParam(Params, TEXT("source_pin"), SourcePinName))
+	{
+		return FECACommandResult::Error(TEXT("Missing required parameter: source_pin"));
+	}
+	if (!GetStringParam(Params, TEXT("target_node_id"), TargetNodeId))
+	{
+		return FECACommandResult::Error(TEXT("Missing required parameter: target_node_id"));
+	}
+	if (!GetStringParam(Params, TEXT("target_pin"), TargetPinName))
+	{
+		return FECACommandResult::Error(TEXT("Missing required parameter: target_pin"));
+	}
+	
+	FString GraphName = TEXT("EventGraph");
+	GetStringParam(Params, TEXT("graph_name"), GraphName, false);
+	
+	UBlueprint* Blueprint = LoadBlueprintByPath(BlueprintPath);
+	if (!Blueprint)
+	{
+		return FECACommandResult::Error(FString::Printf(TEXT("Blueprint not found: %s"), *BlueprintPath));
+	}
+	
+	UEdGraph* Graph = FindGraphByName(Blueprint, GraphName);
+	if (!Graph)
+	{
+		return FECACommandResult::Error(FString::Printf(TEXT("Graph not found: %s"), *GraphName));
+	}
+	
+	// Parse GUIDs
+	FGuid SourceGuid, TargetGuid;
+	if (!FGuid::Parse(SourceNodeId, SourceGuid))
+	{
+		return FECACommandResult::Error(FString::Printf(TEXT("Invalid source node ID: %s"), *SourceNodeId));
+	}
+	if (!FGuid::Parse(TargetNodeId, TargetGuid))
+	{
+		return FECACommandResult::Error(FString::Printf(TEXT("Invalid target node ID: %s"), *TargetNodeId));
+	}
+	
+	// Find nodes
+	UEdGraphNode* SourceNode = FindNodeByGuid(Graph, SourceGuid);
+	UEdGraphNode* TargetNode = FindNodeByGuid(Graph, TargetGuid);
+	
+	if (!SourceNode)
+	{
+		return FECACommandResult::Error(FString::Printf(TEXT("Source node not found: %s"), *SourceNodeId));
+	}
+	if (!TargetNode)
+	{
+		return FECACommandResult::Error(FString::Printf(TEXT("Target node not found: %s"), *TargetNodeId));
+	}
+	
+	// Find pins using friendly name helper (handles "execute"/"then" aliases for unnamed exec pins)
+	UEdGraphPin* SourcePin = FindPinByFriendlyName(SourceNode, SourcePinName, EGPD_Output);
+	UEdGraphPin* TargetPin = FindPinByFriendlyName(TargetNode, TargetPinName, EGPD_Input);
+	
+	if (!SourcePin)
+	{
+		return FECACommandResult::Error(FString::Printf(TEXT("Source pin not found: %s"), *SourcePinName));
+	}
+	if (!TargetPin)
+	{
+		return FECACommandResult::Error(FString::Printf(TEXT("Target pin not found: %s"), *TargetPinName));
+	}
+	
+	// Make the connection
+	SourcePin->MakeLinkTo(TargetPin);
+	
+	// Verify the connection was made
+	if (!SourcePin->LinkedTo.Contains(TargetPin))
+	{
+		return FECACommandResult::Error(TEXT("Failed to connect pins - types may be incompatible"));
+	}
+	
+	// Notify nodes that pins have changed - this triggers type propagation
+	// which is essential for wildcard pins (like ForEachLoop's Array Element)
+	// to infer their type from the connected pin
+	SourceNode->PinConnectionListChanged(SourcePin);
+	TargetNode->PinConnectionListChanged(TargetPin);
+	
+	// For nodes with wildcard pins, reconstruction may be needed to propagate types
+	// Check if target node has any wildcard pins that might need updating
+	bool bTargetNeedsReconstruct = false;
+	for (UEdGraphPin* Pin : TargetNode->Pins)
+	{
+		if (Pin && Pin->PinType.PinCategory == UEdGraphSchema_K2::PC_Wildcard)
+		{
+			bTargetNeedsReconstruct = true;
+			break;
+		}
+	}
+	
+	if (bTargetNeedsReconstruct)
+	{
+		// Use the schema to propagate pin types - this handles wildcard inference
+		const UEdGraphSchema* Schema = Graph->GetSchema();
+		if (Schema)
+		{
+			// NotifyPinConnectionListChanged triggers the type propagation logic
+			Schema->TrySetDefaultValue(*TargetPin, TEXT(""));
+		}
+		
+		// Reconstruct the node to update all dependent pins
+		TargetNode->ReconstructNode();
+	}
+	
+	// Also check source node for wildcards
+	bool bSourceNeedsReconstruct = false;
+	for (UEdGraphPin* Pin : SourceNode->Pins)
+	{
+		if (Pin && Pin->PinType.PinCategory == UEdGraphSchema_K2::PC_Wildcard)
+		{
+			bSourceNeedsReconstruct = true;
+			break;
+		}
+	}
+	
+	if (bSourceNeedsReconstruct)
+	{
+		SourceNode->ReconstructNode();
+	}
+	
+	FBlueprintEditorUtils::MarkBlueprintAsStructurallyModified(Blueprint);
+	
+	TSharedPtr<FJsonObject> Result = MakeResult();
+	Result->SetBoolField(TEXT("connected"), true);
+	Result->SetStringField(TEXT("source_node"), SourceNodeId);
+	Result->SetStringField(TEXT("source_pin"), SourcePinName);
+	Result->SetStringField(TEXT("target_node"), TargetNodeId);
+	Result->SetStringField(TEXT("target_pin"), TargetPinName);
+	
+	return FECACommandResult::Success(Result);
+}
+
+//------------------------------------------------------------------------------
+// AddBlueprintSelfReference
+//------------------------------------------------------------------------------
+
+FECACommandResult FECACommand_AddBlueprintSelfReference::Execute(const TSharedPtr<FJsonObject>& Params)
+{
+	FString BlueprintPath;
+	if (!GetStringParam(Params, TEXT("blueprint_path"), BlueprintPath))
+	{
+		return FECACommandResult::Error(TEXT("Missing required parameter: blueprint_path"));
+	}
+	
+	FString GraphName = TEXT("EventGraph");
+	GetStringParam(Params, TEXT("graph_name"), GraphName, false);
+	
+	UBlueprint* Blueprint = LoadBlueprintByPath(BlueprintPath);
+	if (!Blueprint)
+	{
+		return FECACommandResult::Error(FString::Printf(TEXT("Blueprint not found: %s"), *BlueprintPath));
+	}
+	
+	UEdGraph* Graph = FindGraphByName(Blueprint, GraphName);
+	if (!Graph)
+	{
+		return FECACommandResult::Error(FString::Printf(TEXT("Graph not found: %s"), *GraphName));
+	}
+	
+	FVector2D Position = GetNodePosition(Params, Graph);
+	
+	// Create self reference node
+	UK2Node_Self* SelfNode = NewObject<UK2Node_Self>(Graph);
+	SelfNode->NodePosX = Position.X;
+	SelfNode->NodePosY = Position.Y;
+	
+	Graph->AddNode(SelfNode, false, false);
+	SelfNode->AllocateDefaultPins();
+	EnsureNodeHasValidGuid(SelfNode);
+	
+	FBlueprintEditorUtils::MarkBlueprintAsStructurallyModified(Blueprint);
+	
+	TSharedPtr<FJsonObject> Result = MakeResult();
+	Result->SetStringField(TEXT("node_id"), SelfNode->NodeGuid.ToString());
+	Result->SetObjectField(TEXT("node"), NodeToJson(SelfNode));
+	
+	return FECACommandResult::Success(Result);
+}
+
+//------------------------------------------------------------------------------
+// AddBlueprintComponentReference
+//------------------------------------------------------------------------------
+
+FECACommandResult FECACommand_AddBlueprintComponentReference::Execute(const TSharedPtr<FJsonObject>& Params)
+{
+	FString BlueprintPath;
+	if (!GetStringParam(Params, TEXT("blueprint_path"), BlueprintPath))
+	{
+		return FECACommandResult::Error(TEXT("Missing required parameter: blueprint_path"));
+	}
+	
+	FString ComponentName;
+	if (!GetStringParam(Params, TEXT("component_name"), ComponentName))
+	{
+		return FECACommandResult::Error(TEXT("Missing required parameter: component_name"));
+	}
+	
+	FString GraphName = TEXT("EventGraph");
+	GetStringParam(Params, TEXT("graph_name"), GraphName, false);
+	
+	UBlueprint* Blueprint = LoadBlueprintByPath(BlueprintPath);
+	if (!Blueprint)
+	{
+		return FECACommandResult::Error(FString::Printf(TEXT("Blueprint not found: %s"), *BlueprintPath));
+	}
+	
+	UEdGraph* Graph = FindGraphByName(Blueprint, GraphName);
+	if (!Graph)
+	{
+		return FECACommandResult::Error(FString::Printf(TEXT("Graph not found: %s"), *GraphName));
+	}
+	
+	FVector2D Position = GetNodePosition(Params, Graph);
+	
+	// Create variable get node for the component
+	UK2Node_VariableGet* GetNode = NewObject<UK2Node_VariableGet>(Graph);
+	GetNode->VariableReference.SetSelfMember(FName(*ComponentName));
+	GetNode->NodePosX = Position.X;
+	GetNode->NodePosY = Position.Y;
+	
+	Graph->AddNode(GetNode, false, false);
+	GetNode->AllocateDefaultPins();
+	EnsureNodeHasValidGuid(GetNode);
+	
+	FBlueprintEditorUtils::MarkBlueprintAsStructurallyModified(Blueprint);
+	
+	TSharedPtr<FJsonObject> Result = MakeResult();
+	Result->SetStringField(TEXT("node_id"), GetNode->NodeGuid.ToString());
+	Result->SetStringField(TEXT("component_name"), ComponentName);
+	Result->SetObjectField(TEXT("node"), NodeToJson(GetNode));
+	
+	return FECACommandResult::Success(Result);
+}
+
+//------------------------------------------------------------------------------
+// FindBlueprintNodes
+//------------------------------------------------------------------------------
+
+FECACommandResult FECACommand_FindBlueprintNodes::Execute(const TSharedPtr<FJsonObject>& Params)
+{
+	FString BlueprintPath;
+	if (!GetStringParam(Params, TEXT("blueprint_path"), BlueprintPath))
+	{
+		return FECACommandResult::Error(TEXT("Missing required parameter: blueprint_path"));
+	}
+	
+	FString GraphName = TEXT("EventGraph");
+	GetStringParam(Params, TEXT("graph_name"), GraphName, false);
+	
+	FString NodeClass, NodeTitle;
+	GetStringParam(Params, TEXT("node_class"), NodeClass, false);
+	GetStringParam(Params, TEXT("node_title"), NodeTitle, false);
+	
+	UBlueprint* Blueprint = LoadBlueprintByPath(BlueprintPath);
+	if (!Blueprint)
+	{
+		return FECACommandResult::Error(FString::Printf(TEXT("Blueprint not found: %s"), *BlueprintPath));
+	}
+	
+	UEdGraph* Graph = FindGraphByName(Blueprint, GraphName);
+	if (!Graph)
+	{
+		return FECACommandResult::Error(FString::Printf(TEXT("Graph not found: %s"), *GraphName));
+	}
+	
+	TArray<TSharedPtr<FJsonValue>> NodesArray;
+	
+	for (UEdGraphNode* Node : Graph->Nodes)
+	{
+		bool bMatch = true;
+		
+		// Filter by class
+		if (!NodeClass.IsEmpty() && bMatch)
+		{
+			FString ClassName = Node->GetClass()->GetName();
+			bMatch = ClassName.Contains(NodeClass);
+		}
+		
+		// Filter by title
+		if (!NodeTitle.IsEmpty() && bMatch)
+		{
+			FString Title = Node->GetNodeTitle(ENodeTitleType::FullTitle).ToString();
+			bMatch = Title.Contains(NodeTitle);
+		}
+		
+		if (bMatch)
+		{
+			NodesArray.Add(MakeShared<FJsonValueObject>(NodeToJson(Node)));
+		}
+	}
+	
+	TSharedPtr<FJsonObject> Result = MakeResult();
+	Result->SetArrayField(TEXT("nodes"), NodesArray);
+	Result->SetNumberField(TEXT("count"), NodesArray.Num());
+	
+	return FECACommandResult::Success(Result);
+}
+
+//------------------------------------------------------------------------------
+// AddBlueprintVariableGetNode
+//------------------------------------------------------------------------------
+
+FECACommandResult FECACommand_AddBlueprintVariableGetNode::Execute(const TSharedPtr<FJsonObject>& Params)
+{
+	FString BlueprintPath;
+	if (!GetStringParam(Params, TEXT("blueprint_path"), BlueprintPath))
+	{
+		return FECACommandResult::Error(TEXT("Missing required parameter: blueprint_path"));
+	}
+	
+	FString VariableName;
+	if (!GetStringParam(Params, TEXT("variable_name"), VariableName))
+	{
+		return FECACommandResult::Error(TEXT("Missing required parameter: variable_name"));
+	}
+	
+	FString GraphName = TEXT("EventGraph");
+	GetStringParam(Params, TEXT("graph_name"), GraphName, false);
+	
+	UBlueprint* Blueprint = LoadBlueprintByPath(BlueprintPath);
+	if (!Blueprint)
+	{
+		return FECACommandResult::Error(FString::Printf(TEXT("Blueprint not found: %s"), *BlueprintPath));
+	}
+	
+	UEdGraph* Graph = FindGraphByName(Blueprint, GraphName);
+	if (!Graph)
+	{
+		return FECACommandResult::Error(FString::Printf(TEXT("Graph not found: %s"), *GraphName));
+	}
+	
+	FVector2D Position = GetNodePosition(Params, Graph);
+	
+	UK2Node_VariableGet* GetNode = NewObject<UK2Node_VariableGet>(Graph);
+	GetNode->VariableReference.SetSelfMember(FName(*VariableName));
+	GetNode->NodePosX = Position.X;
+	GetNode->NodePosY = Position.Y;
+	
+	Graph->AddNode(GetNode, false, false);
+	GetNode->AllocateDefaultPins();
+	EnsureNodeHasValidGuid(GetNode);
+	GetNode->ReconstructNode();
+	
+	FBlueprintEditorUtils::MarkBlueprintAsStructurallyModified(Blueprint);
+	
+	TSharedPtr<FJsonObject> Result = MakeResult();
+	Result->SetStringField(TEXT("node_id"), GetNode->NodeGuid.ToString());
+	Result->SetStringField(TEXT("variable_name"), VariableName);
+	Result->SetObjectField(TEXT("node"), NodeToJson(GetNode));
+	AddNodeErrorInfo(Result, GetNode);
+	
+	return FECACommandResult::Success(Result);
+}
+
+//------------------------------------------------------------------------------
+// AddBlueprintVariableSetNode
+//------------------------------------------------------------------------------
+
+FECACommandResult FECACommand_AddBlueprintVariableSetNode::Execute(const TSharedPtr<FJsonObject>& Params)
+{
+	FString BlueprintPath;
+	if (!GetStringParam(Params, TEXT("blueprint_path"), BlueprintPath))
+	{
+		return FECACommandResult::Error(TEXT("Missing required parameter: blueprint_path"));
+	}
+	
+	FString VariableName;
+	if (!GetStringParam(Params, TEXT("variable_name"), VariableName))
+	{
+		return FECACommandResult::Error(TEXT("Missing required parameter: variable_name"));
+	}
+	
+	FString GraphName = TEXT("EventGraph");
+	GetStringParam(Params, TEXT("graph_name"), GraphName, false);
+	
+	UBlueprint* Blueprint = LoadBlueprintByPath(BlueprintPath);
+	if (!Blueprint)
+	{
+		return FECACommandResult::Error(FString::Printf(TEXT("Blueprint not found: %s"), *BlueprintPath));
+	}
+	
+	UEdGraph* Graph = FindGraphByName(Blueprint, GraphName);
+	if (!Graph)
+	{
+		return FECACommandResult::Error(FString::Printf(TEXT("Graph not found: %s"), *GraphName));
+	}
+	
+	FVector2D Position = GetNodePosition(Params, Graph);
+	
+	UK2Node_VariableSet* SetNode = NewObject<UK2Node_VariableSet>(Graph);
+	SetNode->VariableReference.SetSelfMember(FName(*VariableName));
+	SetNode->NodePosX = Position.X;
+	SetNode->NodePosY = Position.Y;
+	
+	Graph->AddNode(SetNode, false, false);
+	SetNode->AllocateDefaultPins();
+	EnsureNodeHasValidGuid(SetNode);
+	SetNode->ReconstructNode();
+	
+	FBlueprintEditorUtils::MarkBlueprintAsStructurallyModified(Blueprint);
+	
+	TSharedPtr<FJsonObject> Result = MakeResult();
+	Result->SetStringField(TEXT("node_id"), SetNode->NodeGuid.ToString());
+	Result->SetStringField(TEXT("variable_name"), VariableName);
+	Result->SetObjectField(TEXT("node"), NodeToJson(SetNode));
+	AddNodeErrorInfo(Result, SetNode);
+	
+	return FECACommandResult::Success(Result);
+}
+
+//------------------------------------------------------------------------------
+// BatchEditBlueprintNodes - Create multiple nodes and wire them atomically
+//------------------------------------------------------------------------------
+
+FECACommandResult FECACommand_BatchEditBlueprintNodes::Execute(const TSharedPtr<FJsonObject>& Params)
+{
+	FString BlueprintPath;
+	if (!GetStringParam(Params, TEXT("blueprint_path"), BlueprintPath))
+	{
+		return FECACommandResult::Error(TEXT("Missing required parameter: blueprint_path"));
+	}
+	
+	FString GraphName = TEXT("EventGraph");
+	GetStringParam(Params, TEXT("graph_name"), GraphName, false);
+	
+	UBlueprint* Blueprint = LoadBlueprintByPath(BlueprintPath);
+	if (!Blueprint)
+	{
+		return FECACommandResult::Error(FString::Printf(TEXT("Blueprint not found: %s"), *BlueprintPath));
+	}
+	
+	UEdGraph* Graph = FindGraphByName(Blueprint, GraphName);
+	if (!Graph)
+	{
+		return FECACommandResult::Error(FString::Printf(TEXT("Graph not found: %s"), *GraphName));
+	}
+	
+	// Get nodes array
+	const TArray<TSharedPtr<FJsonValue>>* NodesArray;
+	if (!GetArrayParam(Params, TEXT("nodes"), NodesArray))
+	{
+		return FECACommandResult::Error(TEXT("Missing required parameter: nodes"));
+	}
+	
+	// Map of temp_id -> created node
+	TMap<FString, UEdGraphNode*> TempIdToNode;
+	
+	// Array for result data
+	TArray<TSharedPtr<FJsonValue>> CreatedNodesArray;
+	TArray<FString> Errors;
+	
+	// Create all nodes first
+	for (const TSharedPtr<FJsonValue>& NodeValue : *NodesArray)
+	{
+		if (NodeValue->Type != EJson::Object)
+		{
+			Errors.Add(TEXT("Invalid node definition - expected object"));
+			continue;
+		}
+		
+		const TSharedPtr<FJsonObject>& NodeDef = NodeValue->AsObject();
+		
+		FString TempId;
+		if (!NodeDef->TryGetStringField(TEXT("temp_id"), TempId))
+		{
+			Errors.Add(TEXT("Node definition missing 'temp_id'"));
+			continue;
+		}
+		
+		FString NodeType;
+		if (!NodeDef->TryGetStringField(TEXT("type"), NodeType))
+		{
+			Errors.Add(FString::Printf(TEXT("Node '%s' missing 'type'"), *TempId));
+			continue;
+		}
+		
+		// Get position - auto-calculate if not provided
+		FVector2D Position(0, 0);
+		const TSharedPtr<FJsonObject>* PosObj;
+		if (NodeDef->TryGetObjectField(TEXT("node_position"), PosObj))
+		{
+			Position.X = (*PosObj)->GetNumberField(TEXT("x"));
+			Position.Y = (*PosObj)->GetNumberField(TEXT("y"));
+		}
+		else
+		{
+			// Auto-calculate position for new node
+			Position = CalculateAutoNodePosition(Graph);
+		}
+		
+		UEdGraphNode* CreatedNode = nullptr;
+		
+		// Create node based on type
+		if (NodeType == TEXT("event"))
+		{
+			FString EventName;
+			if (!NodeDef->TryGetStringField(TEXT("event_name"), EventName))
+			{
+				Errors.Add(FString::Printf(TEXT("Node '%s' of type 'event' missing 'event_name'"), *TempId));
+				continue;
+			}
+			
+			UFunction* OverrideFunc = nullptr;
+			if (Blueprint->ParentClass)
+			{
+				OverrideFunc = Blueprint->ParentClass->FindFunctionByName(*EventName);
+			}
+			
+			UK2Node_Event* EventNode = nullptr;
+			if (OverrideFunc)
+			{
+				EventNode = NewObject<UK2Node_Event>(Graph);
+				EventNode->EventReference.SetFromField<UFunction>(OverrideFunc, false);
+				EventNode->bOverrideFunction = true;
+			}
+			else
+			{
+				UK2Node_CustomEvent* CustomEvent = NewObject<UK2Node_CustomEvent>(Graph);
+				CustomEvent->CustomFunctionName = FName(*EventName);
+				EventNode = CustomEvent;
+			}
+			
+			CreatedNode = EventNode;
+		}
+		else if (NodeType == TEXT("custom_event"))
+		{
+			FString EventName;
+			if (!NodeDef->TryGetStringField(TEXT("event_name"), EventName))
+			{
+				Errors.Add(FString::Printf(TEXT("Node '%s' of type 'custom_event' missing 'event_name'"), *TempId));
+				continue;
+			}
+			
+			UK2Node_CustomEvent* CustomEvent = NewObject<UK2Node_CustomEvent>(Graph);
+			CustomEvent->CustomFunctionName = FName(*EventName);
+			CreatedNode = CustomEvent;
+		}
+		else if (NodeType == TEXT("function"))
+		{
+			FString FunctionName;
+			if (!NodeDef->TryGetStringField(TEXT("function_name"), FunctionName))
+			{
+				Errors.Add(FString::Printf(TEXT("Node '%s' of type 'function' missing 'function_name'"), *TempId));
+				continue;
+			}
+			
+			FString TargetClass;
+			NodeDef->TryGetStringField(TEXT("target_class"), TargetClass);
+			
+			// Find function
+			UFunction* Function = nullptr;
+			TArray<UClass*> ClassesToSearch;
+			
+			if (!TargetClass.IsEmpty())
+			{
+				UClass* SpecificClass = FindObject<UClass>(nullptr, *FString::Printf(TEXT("/Script/Engine.%s"), *TargetClass));
+				if (!SpecificClass)
+				{
+					SpecificClass = FindObject<UClass>(nullptr, *FString::Printf(TEXT("/Script/CoreUObject.%s"), *TargetClass));
+				}
+				if (SpecificClass)
+				{
+					ClassesToSearch.Add(SpecificClass);
+				}
+			}
+			
+			ClassesToSearch.Add(UKismetSystemLibrary::StaticClass());
+			ClassesToSearch.Add(UGameplayStatics::StaticClass());
+			ClassesToSearch.Add(AActor::StaticClass());
+			ClassesToSearch.Add(APawn::StaticClass());
+			ClassesToSearch.Add(ACharacter::StaticClass());
+			ClassesToSearch.Add(Blueprint->ParentClass);
+			
+			for (UClass* Class : ClassesToSearch)
+			{
+				if (Class)
+				{
+					Function = Class->FindFunctionByName(*FunctionName);
+					if (Function) break;
+				}
+			}
+			
+			if (!Function)
+			{
+				Errors.Add(FString::Printf(TEXT("Node '%s': Function not found: %s"), *TempId, *FunctionName));
+				continue;
+			}
+			
+			UK2Node_CallFunction* CallNode = NewObject<UK2Node_CallFunction>(Graph);
+			CallNode->SetFromFunction(Function);
+			CreatedNode = CallNode;
+		}
+		else if (NodeType == TEXT("variable_get"))
+		{
+			FString VariableName;
+			if (!NodeDef->TryGetStringField(TEXT("variable_name"), VariableName))
+			{
+				Errors.Add(FString::Printf(TEXT("Node '%s' of type 'variable_get' missing 'variable_name'"), *TempId));
+				continue;
+			}
+			
+			UK2Node_VariableGet* GetNode = NewObject<UK2Node_VariableGet>(Graph);
+			GetNode->VariableReference.SetSelfMember(FName(*VariableName));
+			CreatedNode = GetNode;
+		}
+		else if (NodeType == TEXT("variable_set"))
+		{
+			FString VariableName;
+			if (!NodeDef->TryGetStringField(TEXT("variable_name"), VariableName))
+			{
+				Errors.Add(FString::Printf(TEXT("Node '%s' of type 'variable_set' missing 'variable_name'"), *TempId));
+				continue;
+			}
+			
+			UK2Node_VariableSet* SetNode = NewObject<UK2Node_VariableSet>(Graph);
+			SetNode->VariableReference.SetSelfMember(FName(*VariableName));
+			CreatedNode = SetNode;
+		}
+		else if (NodeType == TEXT("self"))
+		{
+			UK2Node_Self* SelfNode = NewObject<UK2Node_Self>(Graph);
+			CreatedNode = SelfNode;
+		}
+		else if (NodeType == TEXT("component"))
+		{
+			FString ComponentName;
+			if (!NodeDef->TryGetStringField(TEXT("component_name"), ComponentName))
+			{
+				Errors.Add(FString::Printf(TEXT("Node '%s' of type 'component' missing 'component_name'"), *TempId));
+				continue;
+			}
+			
+			UK2Node_VariableGet* GetNode = NewObject<UK2Node_VariableGet>(Graph);
+			GetNode->VariableReference.SetSelfMember(FName(*ComponentName));
+			CreatedNode = GetNode;
+		}
+		else if (NodeType == TEXT("input_action"))
+		{
+			FString ActionName;
+			if (!NodeDef->TryGetStringField(TEXT("action_name"), ActionName))
+			{
+				Errors.Add(FString::Printf(TEXT("Node '%s' of type 'input_action' missing 'action_name'"), *TempId));
+				continue;
+			}
+			
+			UK2Node_InputAction* InputActionNode = NewObject<UK2Node_InputAction>(Graph);
+			InputActionNode->InputActionName = FName(*ActionName);
+			CreatedNode = InputActionNode;
+		}
+		else if (NodeType == TEXT("cast"))
+		{
+			FString TargetClassName;
+			if (!NodeDef->TryGetStringField(TEXT("target_class"), TargetClassName))
+			{
+				Errors.Add(FString::Printf(TEXT("Node '%s' of type 'cast' missing 'target_class'"), *TempId));
+				continue;
+			}
+			
+			bool bPureCast = false;
+			NodeDef->TryGetBoolField(TEXT("pure"), bPureCast);
+			
+			// Find the target class
+			UClass* TargetClass = nullptr;
+			
+			// Try as a blueprint first
+			UBlueprint* TargetBlueprint = LoadObject<UBlueprint>(nullptr, *TargetClassName);
+			if (TargetBlueprint && TargetBlueprint->GeneratedClass)
+			{
+				TargetClass = TargetBlueprint->GeneratedClass;
+			}
+			
+			if (!TargetClass)
+			{
+				// Try as a native class
+				TArray<FString> ModulesToSearch = {
+					TEXT("/Script/Engine."),
+					TEXT("/Script/CoreUObject."),
+					TEXT("/Script/UMG."),
+					TEXT("/Script/AIModule."),
+					TEXT("")
+				};
+				
+				for (const FString& Module : ModulesToSearch)
+				{
+					TargetClass = FindObject<UClass>(nullptr, *(Module + TargetClassName));
+					if (TargetClass) break;
+				}
+			}
+			
+			if (!TargetClass)
+			{
+				// Search all classes by name
+				for (TObjectIterator<UClass> It; It; ++It)
+				{
+					if (It->GetName().Equals(TargetClassName, ESearchCase::IgnoreCase))
+					{
+						TargetClass = *It;
+						break;
+					}
+				}
+			}
+			
+			if (!TargetClass)
+			{
+				Errors.Add(FString::Printf(TEXT("Node '%s': Target class not found: %s"), *TempId, *TargetClassName));
+				continue;
+			}
+			
+			UK2Node_DynamicCast* CastNode = NewObject<UK2Node_DynamicCast>(Graph);
+			CastNode->TargetType = TargetClass;
+			CastNode->SetPurity(bPureCast);
+			CreatedNode = CastNode;
+		}
+		else if (NodeType == TEXT("macro"))
+		{
+			FString MacroName;
+			if (!NodeDef->TryGetStringField(TEXT("macro_name"), MacroName))
+			{
+				Errors.Add(FString::Printf(TEXT("Node '%s' of type 'macro' missing 'macro_name'"), *TempId));
+				continue;
+			}
+			
+			// Load standard macros library
+			UBlueprint* MacroLibrary = LoadObject<UBlueprint>(nullptr, TEXT("/Engine/EditorBlueprintResources/StandardMacros.StandardMacros"));
+			if (!MacroLibrary)
+			{
+				Errors.Add(FString::Printf(TEXT("Node '%s': Failed to load StandardMacros library"), *TempId));
+				continue;
+			}
+			
+			// Find the macro graph
+			UEdGraph* MacroGraph = nullptr;
+			for (UEdGraph* MacroGraphCandidate : MacroLibrary->MacroGraphs)
+			{
+				if (MacroGraphCandidate && MacroGraphCandidate->GetName().Equals(MacroName, ESearchCase::IgnoreCase))
+				{
+					MacroGraph = MacroGraphCandidate;
+					break;
+				}
+			}
+			
+			if (!MacroGraph)
+			{
+				TArray<FString> AvailableMacros;
+				for (UEdGraph* MacroGraphCandidate : MacroLibrary->MacroGraphs)
+				{
+					if (MacroGraphCandidate) AvailableMacros.Add(MacroGraphCandidate->GetName());
+				}
+				Errors.Add(FString::Printf(TEXT("Node '%s': Macro not found: %s. Available: %s"), 
+					*TempId, *MacroName, *FString::Join(AvailableMacros, TEXT(", "))));
+				continue;
+			}
+			
+			UK2Node_MacroInstance* MacroNode = NewObject<UK2Node_MacroInstance>(Graph);
+			MacroNode->SetMacroGraph(MacroGraph);
+			CreatedNode = MacroNode;
+		}
+		else
+		{
+			Errors.Add(FString::Printf(TEXT("Node '%s': Unknown type '%s'. Valid types: event, custom_event, function, variable_get, variable_set, self, component, input_action, cast, macro"), *TempId, *NodeType));
+			continue;
+		}
+		
+		if (CreatedNode)
+		{
+			CreatedNode->NodePosX = Position.X;
+			CreatedNode->NodePosY = Position.Y;
+			Graph->AddNode(CreatedNode, false, false);
+			CreatedNode->AllocateDefaultPins();
+			EnsureNodeHasValidGuid(CreatedNode);
+			
+			TempIdToNode.Add(TempId, CreatedNode);
+			
+			// Add to result
+			TSharedPtr<FJsonObject> NodeResult = NodeToJson(CreatedNode);
+			NodeResult->SetStringField(TEXT("temp_id"), TempId);
+			CreatedNodesArray.Add(MakeShared<FJsonValueObject>(NodeResult));
+		}
+	}
+	
+	// Now process connections
+	const TArray<TSharedPtr<FJsonValue>>* ConnectionsArray = NULL;
+	TArray<TSharedPtr<FJsonValue>> CompletedConnectionsArray;
+	
+	if (GetArrayParam(Params, TEXT("connections"), ConnectionsArray, false) && ConnectionsArray)
+	{
+		for (const TSharedPtr<FJsonValue>& ConnValue : *ConnectionsArray)
+		{
+			if (ConnValue->Type != EJson::Object)
+			{
+				Errors.Add(TEXT("Invalid connection definition - expected object"));
+				continue;
+			}
+			
+			const TSharedPtr<FJsonObject>& ConnDef = ConnValue->AsObject();
+			
+			FString SourceNodeRef, TargetNodeRef;
+			FString SourcePinName, TargetPinName;
+			
+			if (!ConnDef->TryGetStringField(TEXT("source_node"), SourceNodeRef))
+			{
+				Errors.Add(TEXT("Connection missing 'source_node'"));
+				continue;
+			}
+			if (!ConnDef->TryGetStringField(TEXT("source_pin"), SourcePinName))
+			{
+				Errors.Add(TEXT("Connection missing 'source_pin'"));
+				continue;
+			}
+			if (!ConnDef->TryGetStringField(TEXT("target_node"), TargetNodeRef))
+			{
+				Errors.Add(TEXT("Connection missing 'target_node'"));
+				continue;
+			}
+			if (!ConnDef->TryGetStringField(TEXT("target_pin"), TargetPinName))
+			{
+				Errors.Add(TEXT("Connection missing 'target_pin'"));
+				continue;
+			}
+			
+			// Resolve nodes (can be temp_id or GUID)
+			UEdGraphNode* SourceNode = ResolveNodeReference(Graph, SourceNodeRef, TempIdToNode);
+			UEdGraphNode* TargetNode = ResolveNodeReference(Graph, TargetNodeRef, TempIdToNode);
+			
+			if (!SourceNode)
+			{
+				Errors.Add(FString::Printf(TEXT("Connection: Source node not found: %s"), *SourceNodeRef));
+				continue;
+			}
+			if (!TargetNode)
+			{
+				Errors.Add(FString::Printf(TEXT("Connection: Target node not found: %s"), *TargetNodeRef));
+				continue;
+			}
+			
+			// Find pins using friendly name helper (handles "execute"/"then" aliases)
+			UEdGraphPin* SourcePin = FindPinByFriendlyName(SourceNode, SourcePinName, EGPD_Output);
+			UEdGraphPin* TargetPin = FindPinByFriendlyName(TargetNode, TargetPinName, EGPD_Input);
+			
+			if (!SourcePin)
+			{
+				Errors.Add(FString::Printf(TEXT("Connection: Source pin '%s' not found on node '%s'"), *SourcePinName, *SourceNodeRef));
+				continue;
+			}
+			if (!TargetPin)
+			{
+				Errors.Add(FString::Printf(TEXT("Connection: Target pin '%s' not found on node '%s'"), *TargetPinName, *TargetNodeRef));
+				continue;
+			}
+			
+			// Make connection
+			SourcePin->MakeLinkTo(TargetPin);
+			
+			if (SourcePin->LinkedTo.Contains(TargetPin))
+			{
+				// Notify nodes of connection change - triggers type propagation for wildcards
+				SourceNode->PinConnectionListChanged(SourcePin);
+				TargetNode->PinConnectionListChanged(TargetPin);
+				
+				TSharedPtr<FJsonObject> ConnResult = MakeShared<FJsonObject>();
+				ConnResult->SetStringField(TEXT("source_node"), SourceNode->NodeGuid.ToString());
+				ConnResult->SetStringField(TEXT("source_pin"), SourcePinName);
+				ConnResult->SetStringField(TEXT("target_node"), TargetNode->NodeGuid.ToString());
+				ConnResult->SetStringField(TEXT("target_pin"), TargetPinName);
+				CompletedConnectionsArray.Add(MakeShared<FJsonValueObject>(ConnResult));
+			}
+			else
+			{
+				Errors.Add(FString::Printf(TEXT("Connection failed between %s.%s -> %s.%s (types may be incompatible)"), 
+					*SourceNodeRef, *SourcePinName, *TargetNodeRef, *TargetPinName));
+			}
+		}
+	}
+	
+	// After all connections are made, reconstruct any nodes with wildcard pins
+	// This ensures type propagation happens for nodes like ForEachLoop
+	for (const auto& Pair : TempIdToNode)
+	{
+		UEdGraphNode* Node = Pair.Value;
+		bool bHasWildcard = false;
+		for (UEdGraphPin* Pin : Node->Pins)
+		{
+			if (Pin && Pin->PinType.PinCategory == UEdGraphSchema_K2::PC_Wildcard)
+			{
+				bHasWildcard = true;
+				break;
+			}
+		}
+		if (bHasWildcard)
+		{
+			Node->ReconstructNode();
+		}
+	}
+	
+	FBlueprintEditorUtils::MarkBlueprintAsStructurallyModified(Blueprint);
+	
+	// Build result
+	TSharedPtr<FJsonObject> Result = MakeResult();
+	Result->SetArrayField(TEXT("nodes"), CreatedNodesArray);
+	Result->SetNumberField(TEXT("nodes_created"), CreatedNodesArray.Num());
+	Result->SetArrayField(TEXT("connections"), CompletedConnectionsArray);
+	Result->SetNumberField(TEXT("connections_made"), CompletedConnectionsArray.Num());
+	
+	// Build temp_id -> GUID mapping
+	TSharedPtr<FJsonObject> IdMapping = MakeShared<FJsonObject>();
+	for (const auto& Pair : TempIdToNode)
+	{
+		IdMapping->SetStringField(Pair.Key, Pair.Value->NodeGuid.ToString());
+	}
+	Result->SetObjectField(TEXT("id_mapping"), IdMapping);
+	
+	// Include errors if any
+	if (Errors.Num() > 0)
+	{
+		TArray<TSharedPtr<FJsonValue>> ErrorsArray;
+		for (const FString& Error : Errors)
+		{
+			ErrorsArray.Add(MakeShared<FJsonValueString>(Error));
+		}
+		Result->SetArrayField(TEXT("errors"), ErrorsArray);
+	}
+	
+	return FECACommandResult::Success(Result);
+}
+
+//------------------------------------------------------------------------------
+// GetBlueprintNodePins - Get detailed pin information for a node
+//------------------------------------------------------------------------------
+
+FECACommandResult FECACommand_GetBlueprintNodePins::Execute(const TSharedPtr<FJsonObject>& Params)
+{
+	FString BlueprintPath;
+	if (!GetStringParam(Params, TEXT("blueprint_path"), BlueprintPath))
+	{
+		return FECACommandResult::Error(TEXT("Missing required parameter: blueprint_path"));
+	}
+	
+	FString NodeId;
+	if (!GetStringParam(Params, TEXT("node_id"), NodeId))
+	{
+		return FECACommandResult::Error(TEXT("Missing required parameter: node_id"));
+	}
+	
+	FString GraphName = TEXT("EventGraph");
+	GetStringParam(Params, TEXT("graph_name"), GraphName, false);
+	
+	UBlueprint* Blueprint = LoadBlueprintByPath(BlueprintPath);
+	if (!Blueprint)
+	{
+		return FECACommandResult::Error(FString::Printf(TEXT("Blueprint not found: %s"), *BlueprintPath));
+	}
+	
+	UEdGraph* Graph = FindGraphByName(Blueprint, GraphName);
+	if (!Graph)
+	{
+		return FECACommandResult::Error(FString::Printf(TEXT("Graph not found: %s"), *GraphName));
+	}
+	
+	FGuid NodeGuid;
+	if (!FGuid::Parse(NodeId, NodeGuid))
+	{
+		return FECACommandResult::Error(FString::Printf(TEXT("Invalid node ID: %s"), *NodeId));
+	}
+	
+	UEdGraphNode* Node = FindNodeByGuid(Graph, NodeGuid);
+	if (!Node)
+	{
+		return FECACommandResult::Error(FString::Printf(TEXT("Node not found: %s"), *NodeId));
+	}
+	
+	// Build detailed pin information
+	TArray<TSharedPtr<FJsonValue>> InputPins;
+	TArray<TSharedPtr<FJsonValue>> OutputPins;
+	
+	for (UEdGraphPin* Pin : Node->Pins)
+	{
+		TSharedPtr<FJsonObject> PinInfo = PinToJson(Pin);
+		
+		if (Pin->Direction == EGPD_Input)
+		{
+			InputPins.Add(MakeShared<FJsonValueObject>(PinInfo));
+		}
+		else
+		{
+			OutputPins.Add(MakeShared<FJsonValueObject>(PinInfo));
+		}
+	}
+	
+	TSharedPtr<FJsonObject> Result = MakeResult();
+	Result->SetStringField(TEXT("node_id"), NodeId);
+	Result->SetStringField(TEXT("node_class"), Node->GetClass()->GetName());
+	Result->SetStringField(TEXT("node_title"), Node->GetNodeTitle(ENodeTitleType::FullTitle).ToString());
+	Result->SetArrayField(TEXT("input_pins"), InputPins);
+	Result->SetArrayField(TEXT("output_pins"), OutputPins);
+	Result->SetNumberField(TEXT("input_pin_count"), InputPins.Num());
+	Result->SetNumberField(TEXT("output_pin_count"), OutputPins.Num());
+	
+	return FECACommandResult::Success(Result);
+}
+
+
+//------------------------------------------------------------------------------
+// DeleteBlueprintNode
+//------------------------------------------------------------------------------
+
+FECACommandResult FECACommand_DeleteBlueprintNode::Execute(const TSharedPtr<FJsonObject>& Params)
+{
+	FString BlueprintPath;
+	if (!GetStringParam(Params, TEXT("blueprint_path"), BlueprintPath))
+	{
+		return FECACommandResult::Error(TEXT("Missing required parameter: blueprint_path"));
+	}
+	
+	FString NodeId;
+	if (!GetStringParam(Params, TEXT("node_id"), NodeId))
+	{
+		return FECACommandResult::Error(TEXT("Missing required parameter: node_id"));
+	}
+	
+	FString GraphName = TEXT("EventGraph");
+	GetStringParam(Params, TEXT("graph_name"), GraphName, false);
+	
+	UBlueprint* Blueprint = LoadBlueprintByPath(BlueprintPath);
+	if (!Blueprint)
+	{
+		return FECACommandResult::Error(FString::Printf(TEXT("Failed to load Blueprint: %s"), *BlueprintPath));
+	}
+	
+	UEdGraph* Graph = FindGraphByName(Blueprint, GraphName);
+	if (!Graph)
+	{
+		return FECACommandResult::Error(FString::Printf(TEXT("Graph not found: %s"), *GraphName));
+	}
+	
+	// Parse the GUID
+	FGuid TargetGuid;
+	if (!FGuid::Parse(NodeId, TargetGuid))
+	{
+		return FECACommandResult::Error(FString::Printf(TEXT("Invalid GUID format: %s"), *NodeId));
+	}
+	
+	// Find the node by GUID
+	UEdGraphNode* NodeToDelete = nullptr;
+	for (UEdGraphNode* Node : Graph->Nodes)
+	{
+		if (Node && Node->NodeGuid == TargetGuid)
+		{
+			NodeToDelete = Node;
+			break;
+		}
+	}
+	
+	if (!NodeToDelete)
+	{
+		return FECACommandResult::Error(FString::Printf(TEXT("Node not found with GUID: %s"), *NodeId));
+	}
+	
+	// Store node info before deletion
+	FString DeletedNodeTitle = NodeToDelete->GetNodeTitle(ENodeTitleType::FullTitle).ToString();
+	FString DeletedNodeClass = NodeToDelete->GetClass()->GetName();
+	
+	// Remove the node
+	Graph->RemoveNode(NodeToDelete);
+	
+	// Mark blueprint as modified
+	FBlueprintEditorUtils::MarkBlueprintAsStructurallyModified(Blueprint);
+	
+	TSharedPtr<FJsonObject> Result = MakeResult();
+	Result->SetStringField(TEXT("deleted_node_id"), NodeId);
+	Result->SetStringField(TEXT("deleted_node_title"), DeletedNodeTitle);
+	Result->SetStringField(TEXT("deleted_node_class"), DeletedNodeClass);
+	
+	return FECACommandResult::Success(Result);
+}
+
+//------------------------------------------------------------------------------
+// CleanupOrphanNodes
+//------------------------------------------------------------------------------
+
+FECACommandResult FECACommand_CleanupOrphanNodes::Execute(const TSharedPtr<FJsonObject>& Params)
+{
+	FString BlueprintPath;
+	if (!GetStringParam(Params, TEXT("blueprint_path"), BlueprintPath))
+	{
+		return FECACommandResult::Error(TEXT("Missing required parameter: blueprint_path"));
+	}
+	
+	FString GraphName = TEXT("EventGraph");
+	GetStringParam(Params, TEXT("graph_name"), GraphName, false);
+	
+	bool bDelete = false;
+	GetBoolParam(Params, TEXT("delete"), bDelete, false);
+	
+	FString NodeClassFilter;
+	GetStringParam(Params, TEXT("node_class_filter"), NodeClassFilter, false);
+	
+	UBlueprint* Blueprint = LoadBlueprintByPath(BlueprintPath);
+	if (!Blueprint)
+	{
+		return FECACommandResult::Error(FString::Printf(TEXT("Failed to load Blueprint: %s"), *BlueprintPath));
+	}
+	
+	UEdGraph* Graph = FindGraphByName(Blueprint, GraphName);
+	if (!Graph)
+	{
+		return FECACommandResult::Error(FString::Printf(TEXT("Graph not found: %s"), *GraphName));
+	}
+	
+	// Find orphan nodes - nodes that have exec pins but none are connected
+	TArray<UEdGraphNode*> OrphanNodes;
+	TArray<TSharedPtr<FJsonValue>> OrphanNodesJson;
+	
+	for (UEdGraphNode* Node : Graph->Nodes)
+	{
+		if (!Node)
+		{
+			continue;
+		}
+		
+		// Apply class filter if specified
+		if (!NodeClassFilter.IsEmpty())
+		{
+			if (!Node->GetClass()->GetName().Contains(NodeClassFilter))
+			{
+				continue;
+			}
+		}
+		
+		// Skip certain node types that are always valid (event nodes, entry points)
+		if (Node->IsA<UK2Node_Event>() || Node->IsA<UK2Node_FunctionEntry>() || Node->IsA<UK2Node_FunctionResult>())
+		{
+			continue;
+		}
+		
+		// Check if this node has any exec pins
+		bool bHasExecPins = false;
+		bool bHasConnectedExecPin = false;
+		
+		for (UEdGraphPin* Pin : Node->Pins)
+		{
+			if (Pin && Pin->PinType.PinCategory == UEdGraphSchema_K2::PC_Exec)
+			{
+				bHasExecPins = true;
+				if (Pin->LinkedTo.Num() > 0)
+				{
+					bHasConnectedExecPin = true;
+					break;
+				}
+			}
+		}
+		
+		// If the node has exec pins but none are connected, it's an orphan
+		if (bHasExecPins && !bHasConnectedExecPin)
+		{
+			OrphanNodes.Add(Node);
+			
+			TSharedPtr<FJsonObject> NodeJson = MakeShared<FJsonObject>();
+			NodeJson->SetStringField(TEXT("node_id"), Node->NodeGuid.ToString());
+			NodeJson->SetStringField(TEXT("node_title"), Node->GetNodeTitle(ENodeTitleType::FullTitle).ToString());
+			NodeJson->SetStringField(TEXT("node_class"), Node->GetClass()->GetName());
+			NodeJson->SetNumberField(TEXT("pos_x"), Node->NodePosX);
+			NodeJson->SetNumberField(TEXT("pos_y"), Node->NodePosY);
+			
+			OrphanNodesJson.Add(MakeShared<FJsonValueObject>(NodeJson));
+		}
+	}
+	
+	int32 DeletedCount = 0;
+	
+	// Delete if requested
+	if (bDelete && OrphanNodes.Num() > 0)
+	{
+		for (UEdGraphNode* Node : OrphanNodes)
+		{
+			Graph->RemoveNode(Node);
+			DeletedCount++;
+		}
+		
+		FBlueprintEditorUtils::MarkBlueprintAsStructurallyModified(Blueprint);
+	}
+	
+	TSharedPtr<FJsonObject> Result = MakeResult();
+	Result->SetArrayField(TEXT("orphan_nodes"), OrphanNodesJson);
+	Result->SetNumberField(TEXT("count"), OrphanNodesJson.Num());
+	Result->SetBoolField(TEXT("deleted"), bDelete);
+	Result->SetNumberField(TEXT("deleted_count"), DeletedCount);
+	
+	return FECACommandResult::Success(Result);
+}
+
+//------------------------------------------------------------------------------
+// DisconnectBlueprintNode
+//------------------------------------------------------------------------------
+
+FECACommandResult FECACommand_DisconnectBlueprintNode::Execute(const TSharedPtr<FJsonObject>& Params)
+{
+	FString BlueprintPath;
+	if (!GetStringParam(Params, TEXT("blueprint_path"), BlueprintPath))
+	{
+		return FECACommandResult::Error(TEXT("Missing required parameter: blueprint_path"));
+	}
+	
+	FString NodeId;
+	if (!GetStringParam(Params, TEXT("node_id"), NodeId))
+	{
+		return FECACommandResult::Error(TEXT("Missing required parameter: node_id"));
+	}
+	
+	FString GraphName = TEXT("EventGraph");
+	GetStringParam(Params, TEXT("graph_name"), GraphName, false);
+	
+	FString PinName;
+	bool bSpecificPin = GetStringParam(Params, TEXT("pin_name"), PinName, false);
+	
+	UBlueprint* Blueprint = LoadBlueprintByPath(BlueprintPath);
+	if (!Blueprint)
+	{
+		return FECACommandResult::Error(FString::Printf(TEXT("Failed to load Blueprint: %s"), *BlueprintPath));
+	}
+	
+	UEdGraph* Graph = FindGraphByName(Blueprint, GraphName);
+	if (!Graph)
+	{
+		return FECACommandResult::Error(FString::Printf(TEXT("Graph not found: %s"), *GraphName));
+	}
+	
+	// Parse the GUID
+	FGuid TargetGuid;
+	if (!FGuid::Parse(NodeId, TargetGuid))
+	{
+		return FECACommandResult::Error(FString::Printf(TEXT("Invalid GUID format: %s"), *NodeId));
+	}
+	
+	// Find the node by GUID
+	UEdGraphNode* TargetNode = nullptr;
+	for (UEdGraphNode* Node : Graph->Nodes)
+	{
+		if (Node && Node->NodeGuid == TargetGuid)
+		{
+			TargetNode = Node;
+			break;
+		}
+	}
+	
+	if (!TargetNode)
+	{
+		return FECACommandResult::Error(FString::Printf(TEXT("Node not found with GUID: %s"), *NodeId));
+	}
+	
+	int32 DisconnectedCount = 0;
+	TArray<FString> DisconnectedPins;
+	
+	for (UEdGraphPin* Pin : TargetNode->Pins)
+	{
+		if (!Pin)
+		{
+			continue;
+		}
+		
+		// If a specific pin is requested, only disconnect that one
+		if (bSpecificPin && !Pin->PinName.ToString().Equals(PinName, ESearchCase::IgnoreCase))
+		{
+			continue;
+		}
+		
+		if (Pin->LinkedTo.Num() > 0)
+		{
+			DisconnectedPins.Add(Pin->PinName.ToString());
+			Pin->BreakAllPinLinks();
+			DisconnectedCount++;
+		}
+	}
+	
+	if (DisconnectedCount > 0)
+	{
+		FBlueprintEditorUtils::MarkBlueprintAsStructurallyModified(Blueprint);
+	}
+	
+	TSharedPtr<FJsonObject> Result = MakeResult();
+	Result->SetStringField(TEXT("node_id"), NodeId);
+	Result->SetNumberField(TEXT("disconnected_pin_count"), DisconnectedCount);
+	
+	TArray<TSharedPtr<FJsonValue>> PinsArray;
+	for (const FString& PinNameStr : DisconnectedPins)
+	{
+		PinsArray.Add(MakeShared<FJsonValueString>(PinNameStr));
+	}
+	Result->SetArrayField(TEXT("disconnected_pins"), PinsArray);
+	
+	return FECACommandResult::Success(Result);
+}
+
+
+//------------------------------------------------------------------------------
+// SetBlueprintPinValue - Set the default value of a pin on a node
+//------------------------------------------------------------------------------
+
+FECACommandResult FECACommand_SetBlueprintPinValue::Execute(const TSharedPtr<FJsonObject>& Params)
+{
+	FString BlueprintPath;
+	if (!GetStringParam(Params, TEXT("blueprint_path"), BlueprintPath))
+	{
+		return FECACommandResult::Error(TEXT("Missing required parameter: blueprint_path"));
+	}
+	
+	FString NodeId;
+	if (!GetStringParam(Params, TEXT("node_id"), NodeId))
+	{
+		return FECACommandResult::Error(TEXT("Missing required parameter: node_id"));
+	}
+	
+	FString PinName;
+	if (!GetStringParam(Params, TEXT("pin_name"), PinName))
+	{
+		return FECACommandResult::Error(TEXT("Missing required parameter: pin_name"));
+	}
+	
+	FString Value;
+	if (!GetStringParam(Params, TEXT("value"), Value))
+	{
+		return FECACommandResult::Error(TEXT("Missing required parameter: value"));
+	}
+	
+	FString GraphName = TEXT("EventGraph");
+	GetStringParam(Params, TEXT("graph_name"), GraphName, false);
+	
+	UBlueprint* Blueprint = LoadBlueprintByPath(BlueprintPath);
+	if (!Blueprint)
+	{
+		return FECACommandResult::Error(FString::Printf(TEXT("Blueprint not found: %s"), *BlueprintPath));
+	}
+	
+	UEdGraph* Graph = FindGraphByName(Blueprint, GraphName);
+	if (!Graph)
+	{
+		return FECACommandResult::Error(FString::Printf(TEXT("Graph not found: %s"), *GraphName));
+	}
+	
+	// Parse the GUID
+	FGuid NodeGuid;
+	if (!FGuid::Parse(NodeId, NodeGuid))
+	{
+		return FECACommandResult::Error(FString::Printf(TEXT("Invalid node ID: %s"), *NodeId));
+	}
+	
+	// Find the node
+	UEdGraphNode* Node = FindNodeByGuid(Graph, NodeGuid);
+	if (!Node)
+	{
+		return FECACommandResult::Error(FString::Printf(TEXT("Node not found: %s"), *NodeId));
+	}
+	
+	// Find the pin using friendly name helper (handles "execute"/"then" aliases)
+	UEdGraphPin* Pin = FindPinByFriendlyName(Node, PinName, EGPD_Input);
+	if (!Pin)
+	{
+		// Also try output direction for completeness
+		Pin = FindPinByFriendlyName(Node, PinName, EGPD_Output);
+	}
+	
+	if (!Pin)
+	{
+		return FECACommandResult::Error(FString::Printf(TEXT("Pin not found: %s"), *PinName));
+	}
+	
+	// Check if pin is connected - can't set default on connected pins
+	if (Pin->LinkedTo.Num() > 0)
+	{
+		return FECACommandResult::Error(FString::Printf(TEXT("Cannot set default value on connected pin: %s"), *PinName));
+	}
+	
+	// Store old value for response
+	FString OldValue = Pin->DefaultValue;
+	FString OldObjectPath;
+	if (Pin->DefaultObject)
+	{
+		OldObjectPath = Pin->DefaultObject->GetPathName();
+	}
+	
+	// Determine the pin category to handle different types appropriately
+	FName PinCategory = Pin->PinType.PinCategory;
+	bool bSuccess = false;
+	FString ResultMessage;
+	
+	// Track if we need to reconstruct the node (for class pins that affect output types)
+	bool bNeedsReconstruct = false;
+	
+	// Handle Class reference types (TSubclassOf<> pins) - MUST be handled before generic object pins
+	if (PinCategory == UEdGraphSchema_K2::PC_Class || PinCategory == UEdGraphSchema_K2::PC_SoftClass)
+	{
+		if (Value.IsEmpty())
+		{
+			Pin->DefaultObject = nullptr;
+			Pin->DefaultValue.Empty();
+			bSuccess = true;
+			bNeedsReconstruct = true;
+			ResultMessage = TEXT("Cleared class reference");
+		}
+		else
+		{
+			// For class pins, we need to load the UClass itself
+			UClass* LoadedClass = nullptr;
+			
+			// Try loading as a UClass directly (for native classes like "/Script/Engine.StaticMeshComponent")
+			LoadedClass = LoadObject<UClass>(nullptr, *Value);
+			
+			if (!LoadedClass)
+			{
+				// Try finding by short name in common modules
+				TArray<FString> ModulePaths = {
+					FString::Printf(TEXT("/Script/Engine.%s"), *Value),
+					FString::Printf(TEXT("/Script/CoreUObject.%s"), *Value),
+					FString::Printf(TEXT("/Script/UMG.%s"), *Value),
+					FString::Printf(TEXT("/Script/AIModule.%s"), *Value)
+				};
+				
+				for (const FString& Path : ModulePaths)
+				{
+					LoadedClass = LoadObject<UClass>(nullptr, *Path);
+					if (LoadedClass) break;
+				}
+			}
+			
+			if (!LoadedClass)
+			{
+				// Try loading a Blueprint and getting its generated class
+				UBlueprint* BP = LoadObject<UBlueprint>(nullptr, *Value);
+				if (BP && BP->GeneratedClass)
+				{
+					LoadedClass = BP->GeneratedClass;
+				}
+			}
+			
+			if (!LoadedClass)
+			{
+				// Search all loaded classes by name
+				FString ShortName = FPackageName::GetShortName(Value);
+				for (TObjectIterator<UClass> It; It; ++It)
+				{
+					if (It->GetName().Equals(ShortName, ESearchCase::IgnoreCase))
+					{
+						LoadedClass = *It;
+						break;
+					}
+				}
+			}
+			
+			if (LoadedClass)
+			{
+				// Validate that the class is compatible with the pin's expected type
+				UClass* PinSubCategoryClass = Cast<UClass>(Pin->PinType.PinSubCategoryObject.Get());
+				if (PinSubCategoryClass && !LoadedClass->IsChildOf(PinSubCategoryClass))
+				{
+					return FECACommandResult::Error(FString::Printf(
+						TEXT("Class '%s' is not a subclass of required type '%s'"), 
+						*LoadedClass->GetName(), *PinSubCategoryClass->GetName()));
+				}
+				
+				Pin->DefaultObject = LoadedClass;
+				Pin->DefaultValue.Empty();
+				bSuccess = true;
+				bNeedsReconstruct = true;  // Class pins often affect node outputs
+				ResultMessage = FString::Printf(TEXT("Set class to: %s"), *LoadedClass->GetPathName());
+			}
+			else if (PinCategory == UEdGraphSchema_K2::PC_SoftClass)
+			{
+				// Soft class can store as string path
+				Pin->DefaultValue = Value;
+				bSuccess = true;
+				ResultMessage = FString::Printf(TEXT("Set soft class reference to: %s"), *Value);
+			}
+			else
+			{
+				return FECACommandResult::Error(FString::Printf(TEXT("Failed to find class: %s"), *Value));
+			}
+		}
+	}
+	// Handle Object/Asset reference types
+	else if (PinCategory == UEdGraphSchema_K2::PC_Object ||
+		PinCategory == UEdGraphSchema_K2::PC_SoftObject ||
+		PinCategory == UEdGraphSchema_K2::PC_Interface)
+	{
+		// Value should be an asset path like "/Game/MyFolder/MyAsset.MyAsset"
+		if (Value.IsEmpty())
+		{
+			// Clear the default object
+			Pin->DefaultObject = nullptr;
+			Pin->DefaultValue.Empty();
+			bSuccess = true;
+			ResultMessage = TEXT("Cleared object reference");
+		}
+		else
+		{
+			// Try to load the asset
+			UObject* LoadedObject = StaticLoadObject(UObject::StaticClass(), nullptr, *Value);
+			if (LoadedObject)
+			{
+				Pin->DefaultObject = LoadedObject;
+				Pin->DefaultValue.Empty(); // Clear string value when using object
+				bSuccess = true;
+				ResultMessage = FString::Printf(TEXT("Set object to: %s"), *LoadedObject->GetPathName());
+			}
+			else
+			{
+				// Try as a soft reference path (just store the path string)
+				if (PinCategory == UEdGraphSchema_K2::PC_SoftObject)
+				{
+					Pin->DefaultValue = Value;
+					bSuccess = true;
+					ResultMessage = FString::Printf(TEXT("Set soft reference to: %s"), *Value);
+				}
+				else
+				{
+					return FECACommandResult::Error(FString::Printf(TEXT("Failed to load object: %s"), *Value));
+				}
+			}
+		}
+	}
+	// Handle Text type
+	else if (PinCategory == UEdGraphSchema_K2::PC_Text)
+	{
+		Pin->DefaultTextValue = FText::FromString(Value);
+		Pin->DefaultValue = Value;
+		bSuccess = true;
+		ResultMessage = TEXT("Set text value");
+	}
+	// Handle Struct types (vectors, rotators, etc.)
+	else if (PinCategory == UEdGraphSchema_K2::PC_Struct)
+	{
+		// Struct values are typically stored as strings in a specific format
+		// e.g., Vector: "(X=0,Y=0,Z=0)" or "0,0,0"
+		Pin->DefaultValue = Value;
+		bSuccess = true;
+		ResultMessage = TEXT("Set struct value");
+	}
+	// Handle all other types (bool, int, float, string, name, etc.)
+	else
+	{
+		Pin->DefaultValue = Value;
+		bSuccess = true;
+		ResultMessage = TEXT("Set value");
+	}
+	
+	if (!bSuccess)
+	{
+		return FECACommandResult::Error(TEXT("Failed to set pin value"));
+	}
+	
+	// Notify the node that the pin default changed
+	Node->PinDefaultValueChanged(Pin);
+	
+	// For class pins, reconstruct the node to update dependent pins (e.g., output types)
+	if (bNeedsReconstruct)
+	{
+		Node->ReconstructNode();
+	}
+	
+	// Mark blueprint as modified
+	FBlueprintEditorUtils::MarkBlueprintAsModified(Blueprint);
+	
+	TSharedPtr<FJsonObject> Result = MakeResult();
+	Result->SetStringField(TEXT("node_id"), NodeId);
+	Result->SetStringField(TEXT("pin_name"), PinName);
+	Result->SetStringField(TEXT("old_value"), OldValue);
+	Result->SetStringField(TEXT("old_object"), OldObjectPath);
+	Result->SetStringField(TEXT("new_value"), Value);
+	Result->SetStringField(TEXT("pin_type"), PinCategory.ToString());
+	Result->SetStringField(TEXT("message"), ResultMessage);
+	
+	return FECACommandResult::Success(Result);
+}
+
+//------------------------------------------------------------------------------
+// SetBlueprintVariableDefault - Set the default value of a Blueprint variable
+//------------------------------------------------------------------------------
+
+FECACommandResult FECACommand_SetBlueprintVariableDefault::Execute(const TSharedPtr<FJsonObject>& Params)
+{
+	FString BlueprintPath;
+	if (!GetStringParam(Params, TEXT("blueprint_path"), BlueprintPath))
+	{
+		return FECACommandResult::Error(TEXT("Missing required parameter: blueprint_path"));
+	}
+	
+	FString VariableName;
+	if (!GetStringParam(Params, TEXT("variable_name"), VariableName))
+	{
+		return FECACommandResult::Error(TEXT("Missing required parameter: variable_name"));
+	}
+	
+	FString Value;
+	if (!GetStringParam(Params, TEXT("value"), Value))
+	{
+		return FECACommandResult::Error(TEXT("Missing required parameter: value"));
+	}
+	
+	UBlueprint* Blueprint = LoadBlueprintByPath(BlueprintPath);
+	if (!Blueprint)
+	{
+		return FECACommandResult::Error(FString::Printf(TEXT("Blueprint not found: %s"), *BlueprintPath));
+	}
+	
+	// Find the variable property
+	FName VarName(*VariableName);
+	FProperty* Property = nullptr;
+	int32 VarIndex = INDEX_NONE;
+	
+	for (int32 i = 0; i < Blueprint->NewVariables.Num(); ++i)
+	{
+		if (Blueprint->NewVariables[i].VarName == VarName)
+		{
+			VarIndex = i;
+			break;
+		}
+	}
+	
+	if (VarIndex == INDEX_NONE)
+	{
+		return FECACommandResult::Error(FString::Printf(TEXT("Variable not found: %s"), *VariableName));
+	}
+	
+	FBPVariableDescription& VarDesc = Blueprint->NewVariables[VarIndex];
+	
+	// Store old value for response
+	FString OldValue = VarDesc.DefaultValue;
+	
+	// Set the default value as string - this works for basic types
+	VarDesc.DefaultValue = Value;
+	
+	// For more complex types, we need to modify the CDO (Class Default Object)
+	if (Blueprint->GeneratedClass)
+	{
+		UObject* CDO = Blueprint->GeneratedClass->GetDefaultObject();
+		if (CDO)
+		{
+			Property = Blueprint->GeneratedClass->FindPropertyByName(VarName);
+			if (Property)
+			{
+				// Try to set the property from string
+				void* PropertyAddr = Property->ContainerPtrToValuePtr<void>(CDO);
+				
+				if (FBoolProperty* BoolProp = CastField<FBoolProperty>(Property))
+				{
+					bool bValue = Value.ToBool() || Value.Equals(TEXT("1")) || Value.Equals(TEXT("true"), ESearchCase::IgnoreCase);
+					BoolProp->SetPropertyValue(PropertyAddr, bValue);
+				}
+				else if (FIntProperty* IntProp = CastField<FIntProperty>(Property))
+				{
+					IntProp->SetPropertyValue(PropertyAddr, FCString::Atoi(*Value));
+				}
+				else if (FFloatProperty* FloatProp = CastField<FFloatProperty>(Property))
+				{
+					FloatProp->SetPropertyValue(PropertyAddr, FCString::Atof(*Value));
+				}
+				else if (FDoubleProperty* DoubleProp = CastField<FDoubleProperty>(Property))
+				{
+					DoubleProp->SetPropertyValue(PropertyAddr, FCString::Atod(*Value));
+				}
+				else if (FStrProperty* StrProp = CastField<FStrProperty>(Property))
+				{
+					StrProp->SetPropertyValue(PropertyAddr, Value);
+				}
+				else if (FNameProperty* NameProp = CastField<FNameProperty>(Property))
+				{
+					NameProp->SetPropertyValue(PropertyAddr, FName(*Value));
+				}
+				else if (FTextProperty* TextProp = CastField<FTextProperty>(Property))
+				{
+					TextProp->SetPropertyValue(PropertyAddr, FText::FromString(Value));
+				}
+				else
+				{
+					// For other types, try the generic ImportText
+					Property->ImportText_Direct(*Value, PropertyAddr, CDO, PPF_None);
+				}
+			}
+		}
+	}
+	
+	// Mark blueprint as modified and recompile
+	FBlueprintEditorUtils::MarkBlueprintAsStructurallyModified(Blueprint);
+	
+	TSharedPtr<FJsonObject> Result = MakeResult();
+	Result->SetStringField(TEXT("variable_name"), VariableName);
+	Result->SetStringField(TEXT("old_value"), OldValue);
+	Result->SetStringField(TEXT("new_value"), Value);
+	Result->SetStringField(TEXT("variable_type"), VarDesc.VarType.PinCategory.ToString());
+	
+	return FECACommandResult::Success(Result);
+}
+
+//------------------------------------------------------------------------------
+// GetBlueprintVariableDefault - Get the default value of a Blueprint variable
+//------------------------------------------------------------------------------
+
+FECACommandResult FECACommand_GetBlueprintVariableDefault::Execute(const TSharedPtr<FJsonObject>& Params)
+{
+	FString BlueprintPath;
+	if (!GetStringParam(Params, TEXT("blueprint_path"), BlueprintPath))
+	{
+		return FECACommandResult::Error(TEXT("Missing required parameter: blueprint_path"));
+	}
+	
+	FString VariableName;
+	if (!GetStringParam(Params, TEXT("variable_name"), VariableName))
+	{
+		return FECACommandResult::Error(TEXT("Missing required parameter: variable_name"));
+	}
+	
+	UBlueprint* Blueprint = LoadBlueprintByPath(BlueprintPath);
+	if (!Blueprint)
+	{
+		return FECACommandResult::Error(FString::Printf(TEXT("Blueprint not found: %s"), *BlueprintPath));
+	}
+	
+	// Find the variable
+	FName VarName(*VariableName);
+	int32 VarIndex = INDEX_NONE;
+	
+	for (int32 i = 0; i < Blueprint->NewVariables.Num(); ++i)
+	{
+		if (Blueprint->NewVariables[i].VarName == VarName)
+		{
+			VarIndex = i;
+			break;
+		}
+	}
+	
+	if (VarIndex == INDEX_NONE)
+	{
+		return FECACommandResult::Error(FString::Printf(TEXT("Variable not found: %s"), *VariableName));
+	}
+	
+	FBPVariableDescription& VarDesc = Blueprint->NewVariables[VarIndex];
+	
+	FString DefaultValue = VarDesc.DefaultValue;
+	
+	// Try to get a more accurate value from the CDO
+	if (Blueprint->GeneratedClass)
+	{
+		UObject* CDO = Blueprint->GeneratedClass->GetDefaultObject();
+		if (CDO)
+		{
+			FProperty* Property = Blueprint->GeneratedClass->FindPropertyByName(VarName);
+			if (Property)
+			{
+				void* PropertyAddr = Property->ContainerPtrToValuePtr<void>(CDO);
+				Property->ExportText_Direct(DefaultValue, PropertyAddr, PropertyAddr, CDO, PPF_None);
+			}
+		}
+	}
+	
+	TSharedPtr<FJsonObject> Result = MakeResult();
+	Result->SetStringField(TEXT("variable_name"), VariableName);
+	Result->SetStringField(TEXT("default_value"), DefaultValue);
+	Result->SetStringField(TEXT("variable_type"), VarDesc.VarType.PinCategory.ToString());
+	
+	// Include additional type info
+	if (VarDesc.VarType.PinSubCategoryObject.IsValid())
+	{
+		Result->SetStringField(TEXT("sub_type"), VarDesc.VarType.PinSubCategoryObject->GetName());
+	}
+	Result->SetBoolField(TEXT("is_array"), VarDesc.VarType.IsArray());
+	Result->SetBoolField(TEXT("is_set"), VarDesc.VarType.IsSet());
+	Result->SetBoolField(TEXT("is_map"), VarDesc.VarType.IsMap());
+	
+	return FECACommandResult::Success(Result);
+}
+
+//------------------------------------------------------------------------------
+// BreakPinConnection - Break a specific connection between two pins
+//------------------------------------------------------------------------------
+
+FECACommandResult FECACommand_BreakPinConnection::Execute(const TSharedPtr<FJsonObject>& Params)
+{
+	FString BlueprintPath;
+	if (!GetStringParam(Params, TEXT("blueprint_path"), BlueprintPath))
+	{
+		return FECACommandResult::Error(TEXT("Missing required parameter: blueprint_path"));
+	}
+	
+	FString SourceNodeId;
+	if (!GetStringParam(Params, TEXT("source_node_id"), SourceNodeId))
+	{
+		return FECACommandResult::Error(TEXT("Missing required parameter: source_node_id"));
+	}
+	
+	FString SourcePinName;
+	if (!GetStringParam(Params, TEXT("source_pin"), SourcePinName))
+	{
+		return FECACommandResult::Error(TEXT("Missing required parameter: source_pin"));
+	}
+	
+	FString TargetNodeId;
+	if (!GetStringParam(Params, TEXT("target_node_id"), TargetNodeId))
+	{
+		return FECACommandResult::Error(TEXT("Missing required parameter: target_node_id"));
+	}
+	
+	FString TargetPinName;
+	if (!GetStringParam(Params, TEXT("target_pin"), TargetPinName))
+	{
+		return FECACommandResult::Error(TEXT("Missing required parameter: target_pin"));
+	}
+	
+	FString GraphName = TEXT("EventGraph");
+	GetStringParam(Params, TEXT("graph_name"), GraphName, false);
+	
+	UBlueprint* Blueprint = LoadBlueprintByPath(BlueprintPath);
+	if (!Blueprint)
+	{
+		return FECACommandResult::Error(FString::Printf(TEXT("Blueprint not found: %s"), *BlueprintPath));
+	}
+	
+	UEdGraph* Graph = FindGraphByName(Blueprint, GraphName);
+	if (!Graph)
+	{
+		return FECACommandResult::Error(FString::Printf(TEXT("Graph not found: %s"), *GraphName));
+	}
+	
+	// Parse GUIDs
+	FGuid SourceGuid, TargetGuid;
+	if (!FGuid::Parse(SourceNodeId, SourceGuid))
+	{
+		return FECACommandResult::Error(FString::Printf(TEXT("Invalid source node ID: %s"), *SourceNodeId));
+	}
+	if (!FGuid::Parse(TargetNodeId, TargetGuid))
+	{
+		return FECACommandResult::Error(FString::Printf(TEXT("Invalid target node ID: %s"), *TargetNodeId));
+	}
+	
+	// Find nodes
+	UEdGraphNode* SourceNode = FindNodeByGuid(Graph, SourceGuid);
+	UEdGraphNode* TargetNode = FindNodeByGuid(Graph, TargetGuid);
+	
+	if (!SourceNode)
+	{
+		return FECACommandResult::Error(FString::Printf(TEXT("Source node not found: %s"), *SourceNodeId));
+	}
+	if (!TargetNode)
+	{
+		return FECACommandResult::Error(FString::Printf(TEXT("Target node not found: %s"), *TargetNodeId));
+	}
+	
+	// Find pins using friendly name helper (handles "execute"/"then" aliases)
+	UEdGraphPin* SourcePin = FindPinByFriendlyName(SourceNode, SourcePinName, EGPD_Output);
+	UEdGraphPin* TargetPin = FindPinByFriendlyName(TargetNode, TargetPinName, EGPD_Input);
+	
+	if (!SourcePin)
+	{
+		return FECACommandResult::Error(FString::Printf(TEXT("Source pin not found: %s"), *SourcePinName));
+	}
+	if (!TargetPin)
+	{
+		return FECACommandResult::Error(FString::Printf(TEXT("Target pin not found: %s"), *TargetPinName));
+	}
+	
+	// Check if they're actually connected
+	if (!SourcePin->LinkedTo.Contains(TargetPin))
+	{
+		return FECACommandResult::Error(TEXT("Pins are not connected"));
+	}
+	
+	// Break the connection
+	SourcePin->BreakLinkTo(TargetPin);
+	
+	// Mark blueprint as modified
+	FBlueprintEditorUtils::MarkBlueprintAsStructurallyModified(Blueprint);
+	
+	TSharedPtr<FJsonObject> Result = MakeResult();
+	Result->SetBoolField(TEXT("success"), true);
+	Result->SetStringField(TEXT("source_node"), SourceNodeId);
+	Result->SetStringField(TEXT("source_pin"), SourcePinName);
+	Result->SetStringField(TEXT("target_node"), TargetNodeId);
+	Result->SetStringField(TEXT("target_pin"), TargetPinName);
+	
+	return FECACommandResult::Success(Result);
+}
+
+
+//------------------------------------------------------------------------------
+// AddComponentEventNode - Add an event bound to a specific component
+//------------------------------------------------------------------------------
+
+FECACommandResult FECACommand_AddComponentEventNode::Execute(const TSharedPtr<FJsonObject>& Params)
+{
+	FString BlueprintPath;
+	if (!GetStringParam(Params, TEXT("blueprint_path"), BlueprintPath))
+	{
+		return FECACommandResult::Error(TEXT("Missing required parameter: blueprint_path"));
+	}
+	
+	FString ComponentName;
+	if (!GetStringParam(Params, TEXT("component_name"), ComponentName))
+	{
+		return FECACommandResult::Error(TEXT("Missing required parameter: component_name"));
+	}
+	
+	FString EventName;
+	if (!GetStringParam(Params, TEXT("event_name"), EventName))
+	{
+		return FECACommandResult::Error(TEXT("Missing required parameter: event_name"));
+	}
+	
+	UBlueprint* Blueprint = LoadBlueprintByPath(BlueprintPath);
+	if (!Blueprint)
+	{
+		return FECACommandResult::Error(FString::Printf(TEXT("Blueprint not found: %s"), *BlueprintPath));
+	}
+	
+	// Get the event graph
+	UEdGraph* EventGraph = FindGraphByName(Blueprint, TEXT("EventGraph"));
+	if (!EventGraph)
+	{
+		return FECACommandResult::Error(TEXT("Event graph not found"));
+	}
+	
+	FVector2D Position = GetNodePosition(Params, EventGraph);
+	
+	// Find the component property in the Blueprint
+	FName ComponentPropertyName(*ComponentName);
+	FObjectProperty* ComponentProperty = nullptr;
+	UClass* ComponentClass = nullptr;
+	
+	// Search in the Blueprint's generated class for the component
+	if (Blueprint->GeneratedClass)
+	{
+		for (TFieldIterator<FObjectProperty> PropIt(Blueprint->GeneratedClass); PropIt; ++PropIt)
+		{
+			FObjectProperty* Prop = *PropIt;
+			if (Prop->GetFName() == ComponentPropertyName)
+			{
+				ComponentProperty = Prop;
+				ComponentClass = Prop->PropertyClass;
+				break;
+			}
+		}
+	}
+	
+	// Also search in SimpleConstructionScript for components if not found
+	if (!ComponentClass && Blueprint->SimpleConstructionScript)
+	{
+		TArray<USCS_Node*> AllNodes = Blueprint->SimpleConstructionScript->GetAllNodes();
+		for (USCS_Node* SCSNode : AllNodes)
+		{
+			if (SCSNode && SCSNode->GetVariableName() == ComponentPropertyName)
+			{
+				ComponentClass = SCSNode->ComponentClass;
+				// We still need the property for InitializeComponentBoundEventParams
+				// Try to find it again with the correct name
+				if (Blueprint->GeneratedClass)
+				{
+					ComponentProperty = FindFProperty<FObjectProperty>(Blueprint->GeneratedClass, ComponentPropertyName);
+				}
+				break;
+			}
+		}
+	}
+	
+	if (!ComponentClass)
+	{
+		return FECACommandResult::Error(FString::Printf(TEXT("Component not found: %s"), *ComponentName));
+	}
+	
+	if (!ComponentProperty)
+	{
+		return FECACommandResult::Error(FString::Printf(TEXT("Component property not found for: %s. The Blueprint may need to be compiled first."), *ComponentName));
+	}
+	
+	// Find the delegate/event in the component class
+	FMulticastDelegateProperty* DelegateProperty = nullptr;
+	FName EventFName(*EventName);
+	
+	for (TFieldIterator<FMulticastDelegateProperty> PropIt(ComponentClass, EFieldIteratorFlags::IncludeSuper); PropIt; ++PropIt)
+	{
+		FMulticastDelegateProperty* Prop = *PropIt;
+		if (Prop->GetFName() == EventFName)
+		{
+			DelegateProperty = Prop;
+			break;
+		}
+	}
+	
+	if (!DelegateProperty)
+	{
+		// Build list of available events for error message
+		TArray<FString> AvailableEvents;
+		for (TFieldIterator<FMulticastDelegateProperty> PropIt(ComponentClass, EFieldIteratorFlags::IncludeSuper); PropIt; ++PropIt)
+		{
+			AvailableEvents.Add((*PropIt)->GetName());
+		}
+		
+		return FECACommandResult::Error(FString::Printf(TEXT("Event '%s' not found on component class '%s'. Available events: %s"), 
+			*EventName, *ComponentClass->GetName(), *FString::Join(AvailableEvents, TEXT(", "))));
+	}
+	
+	// Create the component bound event node
+	UK2Node_ComponentBoundEvent* EventNode = NewObject<UK2Node_ComponentBoundEvent>(EventGraph);
+	EventNode->InitializeComponentBoundEventParams(ComponentProperty, DelegateProperty);
+	EventNode->ComponentPropertyName = ComponentPropertyName;
+	EventNode->DelegatePropertyName = EventFName;
+	EventNode->DelegateOwnerClass = ComponentClass;
+	
+	EventNode->NodePosX = Position.X;
+	EventNode->NodePosY = Position.Y;
+	
+	EventGraph->AddNode(EventNode, false, false);
+	EventNode->AllocateDefaultPins();
+	EnsureNodeHasValidGuid(EventNode);
+	
+	FBlueprintEditorUtils::MarkBlueprintAsStructurallyModified(Blueprint);
+	
+	TSharedPtr<FJsonObject> Result = MakeResult();
+	Result->SetStringField(TEXT("node_id"), EventNode->NodeGuid.ToString());
+	Result->SetStringField(TEXT("component_name"), ComponentName);
+	Result->SetStringField(TEXT("event_name"), EventName);
+	Result->SetStringField(TEXT("component_class"), ComponentClass->GetName());
+	Result->SetObjectField(TEXT("node"), NodeToJson(EventNode));
+	
+	return FECACommandResult::Success(Result);
+}
+
+
+//------------------------------------------------------------------------------
+// AddBlueprintMacroNode - Add a macro node (ForLoop, WhileLoop, etc.)
+//------------------------------------------------------------------------------
+
+FECACommandResult FECACommand_AddBlueprintMacroNode::Execute(const TSharedPtr<FJsonObject>& Params)
+{
+	FString BlueprintPath;
+	if (!GetStringParam(Params, TEXT("blueprint_path"), BlueprintPath))
+	{
+		return FECACommandResult::Error(TEXT("Missing required parameter: blueprint_path"));
+	}
+	
+	FString MacroName;
+	if (!GetStringParam(Params, TEXT("macro_name"), MacroName))
+	{
+		return FECACommandResult::Error(TEXT("Missing required parameter: macro_name"));
+	}
+	
+	FString GraphName = TEXT("EventGraph");
+	GetStringParam(Params, TEXT("graph_name"), GraphName, false);
+	
+	UBlueprint* Blueprint = LoadBlueprintByPath(BlueprintPath);
+	if (!Blueprint)
+	{
+		return FECACommandResult::Error(FString::Printf(TEXT("Blueprint not found: %s"), *BlueprintPath));
+	}
+	
+	UEdGraph* Graph = FindGraphByName(Blueprint, GraphName);
+	if (!Graph)
+	{
+		return FECACommandResult::Error(FString::Printf(TEXT("Graph not found: %s"), *GraphName));
+	}
+	
+	FVector2D Position = GetNodePosition(Params, Graph);
+	
+	// Load the standard macros blueprint
+	UBlueprint* MacroLibrary = LoadObject<UBlueprint>(nullptr, TEXT("/Engine/EditorBlueprintResources/StandardMacros.StandardMacros"));
+	if (!MacroLibrary)
+	{
+		return FECACommandResult::Error(TEXT("Failed to load StandardMacros library"));
+	}
+	
+	// Find the macro graph by name
+	UEdGraph* MacroGraph = nullptr;
+	for (UEdGraph* MacroGraphCandidate : MacroLibrary->MacroGraphs)
+	{
+		if (MacroGraphCandidate && MacroGraphCandidate->GetName().Equals(MacroName, ESearchCase::IgnoreCase))
+		{
+			MacroGraph = MacroGraphCandidate;
+			break;
+		}
+	}
+	
+	if (!MacroGraph)
+	{
+		// Build list of available macros for error message
+		TArray<FString> AvailableMacros;
+		for (UEdGraph* MacroGraphCandidate : MacroLibrary->MacroGraphs)
+		{
+			if (MacroGraphCandidate)
+			{
+				AvailableMacros.Add(MacroGraphCandidate->GetName());
+			}
+		}
+		return FECACommandResult::Error(FString::Printf(TEXT("Macro not found: %s. Available macros: %s"), 
+			*MacroName, *FString::Join(AvailableMacros, TEXT(", "))));
+	}
+	
+	// Create the macro instance node
+	UK2Node_MacroInstance* MacroNode = NewObject<UK2Node_MacroInstance>(Graph);
+	MacroNode->SetMacroGraph(MacroGraph);
+	MacroNode->NodePosX = Position.X;
+	MacroNode->NodePosY = Position.Y;
+	
+	Graph->AddNode(MacroNode, false, false);
+	MacroNode->AllocateDefaultPins();
+	EnsureNodeHasValidGuid(MacroNode);
+	MacroNode->ReconstructNode();
+	
+	FBlueprintEditorUtils::MarkBlueprintAsStructurallyModified(Blueprint);
+	
+	TSharedPtr<FJsonObject> Result = MakeResult();
+	Result->SetStringField(TEXT("node_id"), MacroNode->NodeGuid.ToString());
+	Result->SetStringField(TEXT("macro_name"), MacroName);
+	Result->SetObjectField(TEXT("node"), NodeToJson(MacroNode));
+	AddNodeErrorInfo(Result, MacroNode);
+	
+	return FECACommandResult::Success(Result);
+}
+
+//------------------------------------------------------------------------------
+// AddBlueprintCastNode - Add a cast node (Cast To X)
+//------------------------------------------------------------------------------
+
+FECACommandResult FECACommand_AddBlueprintCastNode::Execute(const TSharedPtr<FJsonObject>& Params)
+{
+	FString BlueprintPath;
+	if (!GetStringParam(Params, TEXT("blueprint_path"), BlueprintPath))
+	{
+		return FECACommandResult::Error(TEXT("Missing required parameter: blueprint_path"));
+	}
+	
+	FString TargetClassName;
+	if (!GetStringParam(Params, TEXT("target_class"), TargetClassName))
+	{
+		return FECACommandResult::Error(TEXT("Missing required parameter: target_class"));
+	}
+	
+	FString GraphName = TEXT("EventGraph");
+	GetStringParam(Params, TEXT("graph_name"), GraphName, false);
+	
+	bool bPureCast = false;
+	GetBoolParam(Params, TEXT("pure"), bPureCast, false);
+	
+	UBlueprint* Blueprint = LoadBlueprintByPath(BlueprintPath);
+	if (!Blueprint)
+	{
+		return FECACommandResult::Error(FString::Printf(TEXT("Blueprint not found: %s"), *BlueprintPath));
+	}
+	
+	UEdGraph* Graph = FindGraphByName(Blueprint, GraphName);
+	if (!Graph)
+	{
+		return FECACommandResult::Error(FString::Printf(TEXT("Graph not found: %s"), *GraphName));
+	}
+	
+	FVector2D Position = GetNodePosition(Params, Graph);
+	
+	// Find the target class
+	UClass* TargetClass = nullptr;
+	
+	// Try different search paths for the class
+	// First, try as a blueprint
+	UBlueprint* TargetBlueprint = LoadObject<UBlueprint>(nullptr, *TargetClassName);
+	if (TargetBlueprint && TargetBlueprint->GeneratedClass)
+	{
+		TargetClass = TargetBlueprint->GeneratedClass;
+	}
+	
+	if (!TargetClass)
+	{
+		// Try as a native class in common modules
+		TArray<FString> ModulesToSearch = {
+			TEXT("/Script/Engine."),
+			TEXT("/Script/CoreUObject."),
+			TEXT("/Script/UMG."),
+			TEXT("/Script/AIModule."),
+			TEXT("/Script/NavigationSystem."),
+			TEXT("/Script/PhysicsCore."),
+			TEXT("")  // Try without prefix
+		};
+		
+		for (const FString& Module : ModulesToSearch)
+		{
+			FString FullPath = Module + TargetClassName;
+			TargetClass = FindObject<UClass>(nullptr, *FullPath);
+			if (TargetClass)
+			{
+				break;
+			}
+		}
+	}
+	
+	if (!TargetClass)
+	{
+		// Try searching all classes
+		for (TObjectIterator<UClass> It; It; ++It)
+		{
+			if (It->GetName().Equals(TargetClassName, ESearchCase::IgnoreCase))
+			{
+				TargetClass = *It;
+				break;
+			}
+		}
+	}
+	
+	if (!TargetClass)
+	{
+		return FECACommandResult::Error(FString::Printf(TEXT("Target class not found: %s. For Blueprints, use full path like /Game/Blueprints/BP_MyActor"), *TargetClassName));
+	}
+	
+	// Create the cast node
+	UK2Node_DynamicCast* CastNode = NewObject<UK2Node_DynamicCast>(Graph);
+	CastNode->TargetType = TargetClass;
+	CastNode->SetPurity(bPureCast);
+	CastNode->NodePosX = Position.X;
+	CastNode->NodePosY = Position.Y;
+	
+	Graph->AddNode(CastNode, false, false);
+	CastNode->AllocateDefaultPins();
+	EnsureNodeHasValidGuid(CastNode);
+	CastNode->ReconstructNode();
+	
+	FBlueprintEditorUtils::MarkBlueprintAsStructurallyModified(Blueprint);
+	
+	TSharedPtr<FJsonObject> Result = MakeResult();
+	Result->SetStringField(TEXT("node_id"), CastNode->NodeGuid.ToString());
+	Result->SetStringField(TEXT("target_class"), TargetClass->GetName());
+	Result->SetBoolField(TEXT("pure"), bPureCast);
+	Result->SetObjectField(TEXT("node"), NodeToJson(CastNode));
+	AddNodeErrorInfo(Result, CastNode);
+	
+	return FECACommandResult::Success(Result);
+}
+
+//------------------------------------------------------------------------------
+// DeleteBlueprintComponent - Delete a component from a Blueprint
+//------------------------------------------------------------------------------
+
+FECACommandResult FECACommand_DeleteBlueprintComponent::Execute(const TSharedPtr<FJsonObject>& Params)
+{
+	FString BlueprintPath;
+	if (!GetStringParam(Params, TEXT("blueprint_path"), BlueprintPath))
+	{
+		return FECACommandResult::Error(TEXT("Missing required parameter: blueprint_path"));
+	}
+	
+	FString ComponentName;
+	if (!GetStringParam(Params, TEXT("component_name"), ComponentName))
+	{
+		return FECACommandResult::Error(TEXT("Missing required parameter: component_name"));
+	}
+	
+	UBlueprint* Blueprint = LoadBlueprintByPath(BlueprintPath);
+	if (!Blueprint)
+	{
+		return FECACommandResult::Error(FString::Printf(TEXT("Blueprint not found: %s"), *BlueprintPath));
+	}
+	
+	if (!Blueprint->SimpleConstructionScript)
+	{
+		return FECACommandResult::Error(TEXT("Blueprint has no SimpleConstructionScript"));
+	}
+	
+	// Find the SCS node for this component
+	FName ComponentFName(*ComponentName);
+	USCS_Node* NodeToRemove = nullptr;
+	FString ComponentClass;
+	
+	TArray<USCS_Node*> AllNodes = Blueprint->SimpleConstructionScript->GetAllNodes();
+	for (USCS_Node* SCSNode : AllNodes)
+	{
+		if (SCSNode && SCSNode->GetVariableName() == ComponentFName)
+		{
+			NodeToRemove = SCSNode;
+			if (SCSNode->ComponentClass)
+			{
+				ComponentClass = SCSNode->ComponentClass->GetName();
+			}
+			break;
+		}
+	}
+	
+	if (!NodeToRemove)
+	{
+		// Build list of available components for error message
+		TArray<FString> AvailableComponents;
+		for (USCS_Node* SCSNode : AllNodes)
+		{
+			if (SCSNode)
+			{
+				AvailableComponents.Add(SCSNode->GetVariableName().ToString());
+			}
+		}
+		
+		return FECACommandResult::Error(FString::Printf(TEXT("Component not found: %s. Available components: %s"), 
+			*ComponentName, *FString::Join(AvailableComponents, TEXT(", "))));
+	}
+	
+	// Check if this node has children
+	TArray<USCS_Node*> ChildNodes = NodeToRemove->GetChildNodes();
+	bool bHadChildren = ChildNodes.Num() > 0;
+	
+	// Remove the node from the SCS
+	// This will also handle reparenting children to the parent node if needed
+	Blueprint->SimpleConstructionScript->RemoveNode(NodeToRemove);
+	
+	// Mark blueprint as modified
+	FBlueprintEditorUtils::MarkBlueprintAsStructurallyModified(Blueprint);
+	
+	TSharedPtr<FJsonObject> Result = MakeResult();
+	Result->SetStringField(TEXT("deleted_component"), ComponentName);
+	Result->SetStringField(TEXT("component_class"), ComponentClass);
+	Result->SetBoolField(TEXT("had_children"), bHadChildren);
+	Result->SetNumberField(TEXT("child_count"), ChildNodes.Num());
+	
+	return FECACommandResult::Success(Result);
+}
+
+//------------------------------------------------------------------------------
+// AddBlueprintFlowControlNode
+//------------------------------------------------------------------------------
+
+FECACommandResult FECACommand_AddBlueprintFlowControlNode::Execute(const TSharedPtr<FJsonObject>& Params)
+{
+	FString BlueprintPath;
+	if (!GetStringParam(Params, TEXT("blueprint_path"), BlueprintPath))
+	{
+		return FECACommandResult::Error(TEXT("Missing required parameter: blueprint_path"));
+	}
+	
+	FString NodeType;
+	if (!GetStringParam(Params, TEXT("node_type"), NodeType))
+	{
+		return FECACommandResult::Error(TEXT("Missing required parameter: node_type"));
+	}
+	
+	FString GraphName = TEXT("EventGraph");
+	GetStringParam(Params, TEXT("graph_name"), GraphName, false);
+	
+	FString EnumType;
+	GetStringParam(Params, TEXT("enum_type"), EnumType, false);
+	
+	UBlueprint* Blueprint = LoadBlueprintByPath(BlueprintPath);
+	if (!Blueprint)
+	{
+		return FECACommandResult::Error(FString::Printf(TEXT("Blueprint not found: %s"), *BlueprintPath));
+	}
+	
+	UEdGraph* Graph = FindGraphByName(Blueprint, GraphName);
+	if (!Graph)
+	{
+		return FECACommandResult::Error(FString::Printf(TEXT("Graph not found: %s"), *GraphName));
+	}
+	
+	FVector2D Position = GetNodePosition(Params, Graph);
+	
+	UK2Node* CreatedNode = nullptr;
+	FString NodeTypeName;
+	
+	NodeType = NodeType.ToLower();
+	
+	if (NodeType == TEXT("branch") || NodeType == TEXT("if") || NodeType == TEXT("ifthenelse"))
+	{
+		UK2Node_IfThenElse* BranchNode = NewObject<UK2Node_IfThenElse>(Graph);
+		CreatedNode = BranchNode;
+		NodeTypeName = TEXT("Branch");
+	}
+	else if (NodeType == TEXT("select"))
+	{
+		UK2Node_Select* SelectNode = NewObject<UK2Node_Select>(Graph);
+		CreatedNode = SelectNode;
+		NodeTypeName = TEXT("Select");
+	}
+	else if (NodeType == TEXT("switch_int") || NodeType == TEXT("switchint") || NodeType == TEXT("switch_integer"))
+	{
+		UK2Node_SwitchInteger* SwitchNode = NewObject<UK2Node_SwitchInteger>(Graph);
+		CreatedNode = SwitchNode;
+		NodeTypeName = TEXT("Switch on Int");
+	}
+	else if (NodeType == TEXT("switch_string") || NodeType == TEXT("switchstring"))
+	{
+		UK2Node_SwitchString* SwitchNode = NewObject<UK2Node_SwitchString>(Graph);
+		CreatedNode = SwitchNode;
+		NodeTypeName = TEXT("Switch on String");
+	}
+	else if (NodeType == TEXT("switch_name") || NodeType == TEXT("switchname"))
+	{
+		UK2Node_SwitchName* SwitchNode = NewObject<UK2Node_SwitchName>(Graph);
+		CreatedNode = SwitchNode;
+		NodeTypeName = TEXT("Switch on Name");
+	}
+	else if (NodeType == TEXT("switch_enum") || NodeType == TEXT("switchenum"))
+	{
+		if (EnumType.IsEmpty())
+		{
+			return FECACommandResult::Error(TEXT("switch_enum requires enum_type parameter (e.g., /Script/Engine.ECollisionChannel)"));
+		}
+		
+		UEnum* Enum = FindObject<UEnum>(nullptr, *EnumType);
+		if (!Enum)
+		{
+			return FECACommandResult::Error(FString::Printf(TEXT("Enum type not found: %s"), *EnumType));
+		}
+		
+		UK2Node_SwitchEnum* SwitchNode = NewObject<UK2Node_SwitchEnum>(Graph);
+		SwitchNode->SetEnum(Enum);
+		CreatedNode = SwitchNode;
+		NodeTypeName = FString::Printf(TEXT("Switch on %s"), *Enum->GetName());
+	}
+	else
+	{
+		return FECACommandResult::Error(FString::Printf(TEXT("Unknown flow control node type: %s. Valid types: branch, select, switch_int, switch_string, switch_name, switch_enum"), *NodeType));
+	}
+	
+	if (CreatedNode)
+	{
+		CreatedNode->NodePosX = Position.X;
+		CreatedNode->NodePosY = Position.Y;
+		
+		Graph->AddNode(CreatedNode, false, false);
+		CreatedNode->AllocateDefaultPins();
+		EnsureNodeHasValidGuid(CreatedNode);
+		CreatedNode->ReconstructNode();
+		
+		FBlueprintEditorUtils::MarkBlueprintAsStructurallyModified(Blueprint);
+		
+		TSharedPtr<FJsonObject> Result = MakeResult();
+		Result->SetStringField(TEXT("node_id"), CreatedNode->NodeGuid.ToString());
+		Result->SetStringField(TEXT("node_type"), NodeTypeName);
+		Result->SetObjectField(TEXT("node"), NodeToJson(CreatedNode));
+		
+		// Include any errors
+		AddNodeErrorInfo(Result, CreatedNode);
+		
+		return FECACommandResult::Success(Result);
+	}
+	
+	return FECACommandResult::Error(TEXT("Failed to create flow control node"));
+}
+
+
+//------------------------------------------------------------------------------
+// AutoLayoutBlueprintGraph - Automatically arrange nodes in a Blueprint graph
+// 
+// Uses an improved layout algorithm that:
+// - Properly handles branch nodes (if/else) with parallel lanes
+// - Positions pure nodes near their consumers
+// - Supports multiple layout strategies
+// - Maintains execution flow readability
+//------------------------------------------------------------------------------
+
+FECACommandResult FECACommand_AutoLayoutBlueprintGraph::Execute(const TSharedPtr<FJsonObject>& Params)
+{
+	FString BlueprintPath;
+	if (!GetStringParam(Params, TEXT("blueprint_path"), BlueprintPath))
+	{
+		return FECACommandResult::Error(TEXT("Missing required parameter: blueprint_path"));
+	}
+	
+	UBlueprint* Blueprint = LoadBlueprintByPath(BlueprintPath);
+	if (!Blueprint)
+	{
+		return FECACommandResult::Error(FString::Printf(TEXT("Blueprint not found: %s"), *BlueprintPath));
+	}
+	
+	FString GraphName = TEXT("EventGraph");
+	GetStringParam(Params, TEXT("graph_name"), GraphName, false);
+	
+	UEdGraph* Graph = FindGraphByName(Blueprint, GraphName);
+	if (!Graph)
+	{
+		return FECACommandResult::Error(FString::Printf(TEXT("Graph not found: %s"), *GraphName));
+	}
+	
+	// Get layout configuration (uses padding-based system now)
+	FBlueprintLayoutConfig Config;
+	GetIntParam(Params, TEXT("padding_x"), Config.NodePaddingX, false);
+	GetIntParam(Params, TEXT("padding_y"), Config.NodePaddingY, false);
+	GetIntParam(Params, TEXT("branch_padding"), Config.BranchExtraPaddingY, false);
+	GetIntParam(Params, TEXT("root_padding"), Config.RootExtraPaddingY, false);
+	GetIntParam(Params, TEXT("max_pure_per_column"), Config.MaxPureNodesPerColumn, false);
+	
+	// Get starting position
+	int32 StartX = 0;
+	int32 StartY = 0;
+	GetIntParam(Params, TEXT("start_x"), StartX, false);
+	GetIntParam(Params, TEXT("start_y"), StartY, false);
+	
+	// Check for legacy strategy parameter - log warning if old strategies used
+	FString Strategy;
+	if (GetStringParam(Params, TEXT("strategy"), Strategy, false))
+	{
+		if (Strategy == TEXT("vertical") || Strategy == TEXT("compact"))
+		{
+			// These old strategies are deprecated - use the new algorithm
+			UE_LOG(LogTemp, Warning, TEXT("AutoLayoutBlueprintGraph: '%s' strategy is deprecated. Using improved tree-based layout."), *Strategy);
+		}
+	}
+	
+	// Get specific node IDs if provided (for partial layout)
+	TArray<UEdGraphNode*> SpecificNodes;
+	const TArray<TSharedPtr<FJsonValue>>* NodeIdsArray;
+	if (GetArrayParam(Params, TEXT("node_ids"), NodeIdsArray, false) && NodeIdsArray)
+	{
+		for (const TSharedPtr<FJsonValue>& IdValue : *NodeIdsArray)
+		{
+			FString IdStr;
+			if (IdValue->TryGetString(IdStr))
+			{
+				FGuid Guid;
+				if (FGuid::Parse(IdStr, Guid))
+				{
+					for (UEdGraphNode* Node : Graph->Nodes)
+					{
+						if (Node && Node->NodeGuid == Guid)
+						{
+							SpecificNodes.Add(Node);
+							break;
+						}
+					}
+				}
+			}
+		}
+	}
+	
+	// Run the layout algorithm
+	FBlueprintAutoLayout Layouter(Config);
+	int32 NodesPositioned = 0;
+	
+	if (SpecificNodes.Num() > 0)
+	{
+		NodesPositioned = Layouter.LayoutSubtree(Graph, SpecificNodes, StartX, StartY);
+	}
+	else
+	{
+		NodesPositioned = Layouter.LayoutGraph(Graph, StartX, StartY);
+	}
+	
+	// Mark blueprint as modified
+	FBlueprintEditorUtils::MarkBlueprintAsStructurallyModified(Blueprint);
+	
+	// Build result
+	TSharedPtr<FJsonObject> Result = MakeResult();
+	Result->SetStringField(TEXT("graph_name"), GraphName);
+	Result->SetNumberField(TEXT("nodes_positioned"), NodesPositioned);
+	Result->SetNumberField(TEXT("padding_x"), Config.NodePaddingX);
+	Result->SetNumberField(TEXT("padding_y"), Config.NodePaddingY);
+	
+	// Include positioned node info
+	TArray<TSharedPtr<FJsonValue>> PositionedNodesArray;
+	const TMap<UEdGraphNode*, FLayoutNodeInfo>& LayoutInfo = Layouter.GetLayoutInfo();
+	
+	for (const auto& Pair : LayoutInfo)
+	{
+		UEdGraphNode* Node = Pair.Key;
+		const FLayoutNodeInfo& Info = Pair.Value;
+		
+		TSharedPtr<FJsonObject> NodeInfo = MakeShared<FJsonObject>();
+		NodeInfo->SetStringField(TEXT("node_id"), Node->NodeGuid.ToString());
+		NodeInfo->SetStringField(TEXT("node_title"), Node->GetNodeTitle(ENodeTitleType::FullTitle).ToString());
+		NodeInfo->SetNumberField(TEXT("x"), Node->NodePosX);
+		NodeInfo->SetNumberField(TEXT("y"), Node->NodePosY);
+		NodeInfo->SetNumberField(TEXT("depth"), Info.Depth);
+		NodeInfo->SetNumberField(TEXT("subtree_height"), Info.SubtreeHeight);
+		NodeInfo->SetBoolField(TEXT("is_pure"), Info.bIsPureNode);
+		NodeInfo->SetBoolField(TEXT("is_branch"), Info.bIsBranchNode);
+		NodeInfo->SetBoolField(TEXT("is_root"), Info.bIsRootNode);
+		PositionedNodesArray.Add(MakeShared<FJsonValueObject>(NodeInfo));
+	}
+	Result->SetArrayField(TEXT("positioned_nodes"), PositionedNodesArray);
+	
+	return FECACommandResult::Success(Result);
+}
+
+//------------------------------------------------------------------------------
+// LEGACY: Old AutoLayoutBlueprintGraph implementation (kept for reference)
+// This code is no longer used but preserved for comparison
+//------------------------------------------------------------------------------
+#if 0
+FECACommandResult FECACommand_AutoLayoutBlueprintGraph_Legacy(const TSharedPtr<FJsonObject>& Params)
+{
+	FString BlueprintPath;
+	if (!GetStringParam(Params, TEXT("blueprint_path"), BlueprintPath))
+	{
+		return FECACommandResult::Error(TEXT("Missing required parameter: blueprint_path"));
+	}
+	
+	UBlueprint* Blueprint = LoadBlueprintByPath(BlueprintPath);
+	if (!Blueprint)
+	{
+		return FECACommandResult::Error(FString::Printf(TEXT("Blueprint not found: %s"), *BlueprintPath));
+	}
+	
+	FString GraphName = TEXT("EventGraph");
+	GetStringParam(Params, TEXT("graph_name"), GraphName, false);
+	
+	UEdGraph* Graph = FindGraphByName(Blueprint, GraphName);
+	if (!Graph)
+	{
+		return FECACommandResult::Error(FString::Printf(TEXT("Graph not found: %s"), *GraphName));
+	}
+	
+	// Get layout parameters
+	FString Strategy = TEXT("horizontal");
+	GetStringParam(Params, TEXT("strategy"), Strategy, false);
+	Strategy = Strategy.ToLower();
+	
+	int32 SpacingX = 400;
+	int32 SpacingY = 150;
+	GetIntParam(Params, TEXT("spacing_x"), SpacingX, false);
+	GetIntParam(Params, TEXT("spacing_y"), SpacingY, false);
+	
+	bool bAlignComments = true;
+	GetBoolParam(Params, TEXT("align_comments"), bAlignComments, false);
+	
+	bool bSelectedOnly = false;
+	GetBoolParam(Params, TEXT("selected_only"), bSelectedOnly, false);
+	
+	// Get specific node IDs if provided
+	TSet<FGuid> SpecificNodeIds;
+	const TArray<TSharedPtr<FJsonValue>>* NodeIdsArray;
+	if (GetArrayParam(Params, TEXT("node_ids"), NodeIdsArray, false) && NodeIdsArray)
+	{
+		for (const TSharedPtr<FJsonValue>& IdValue : *NodeIdsArray)
+		{
+			FString IdStr;
+			if (IdValue->TryGetString(IdStr))
+			{
+				FGuid Guid;
+				if (FGuid::Parse(IdStr, Guid))
+				{
+					SpecificNodeIds.Add(Guid);
+				}
+			}
+		}
+	}
+	
+	// Collect nodes to layout
+	TArray<UEdGraphNode*> NodesToLayout;
+	for (UEdGraphNode* Node : Graph->Nodes)
+	{
+		if (!Node)
+		{
+			continue;
+		}
+		
+		// Filter by specific node IDs if provided
+		if (SpecificNodeIds.Num() > 0)
+		{
+			if (!SpecificNodeIds.Contains(Node->NodeGuid))
+			{
+				continue;
+			}
+		}
+		// Filter by selection if requested (only works in editor context)
+		else if (bSelectedOnly)
+		{
+			// In a command context we can't check selection, skip this filter
+			// Users should use node_ids for specific nodes
+		}
+		
+		NodesToLayout.Add(Node);
+	}
+	
+	if (NodesToLayout.Num() == 0)
+	{
+		return FECACommandResult::Error(TEXT("No nodes to layout"));
+	}
+	
+	// Identify root nodes (events, entry points - nodes with no input exec connections)
+	TArray<UEdGraphNode*> RootNodes;
+	TMap<UEdGraphNode*, int32> NodeDepths;
+	TMap<UEdGraphNode*, int32> NodeLanes; // Vertical position index for parallel branches
+	
+	for (UEdGraphNode* Node : NodesToLayout)
+	{
+		bool bIsRoot = false;
+		
+		// Check if this is an event or entry node
+		if (Node->IsA<UK2Node_Event>() || Node->IsA<UK2Node_FunctionEntry>())
+		{
+			bIsRoot = true;
+		}
+		else
+		{
+			// Check if it has exec input pins with no connections
+			bool bHasExecInput = false;
+			bool bHasConnectedExecInput = false;
+			
+			for (UEdGraphPin* Pin : Node->Pins)
+			{
+				if (Pin && Pin->Direction == EGPD_Input && 
+					Pin->PinType.PinCategory == UEdGraphSchema_K2::PC_Exec)
+				{
+					bHasExecInput = true;
+					if (Pin->LinkedTo.Num() > 0)
+					{
+						// Check if connected node is in our layout set
+						for (UEdGraphPin* LinkedPin : Pin->LinkedTo)
+						{
+							if (LinkedPin && NodesToLayout.Contains(LinkedPin->GetOwningNode()))
+							{
+								bHasConnectedExecInput = true;
+								break;
+							}
+						}
+					}
+				}
+			}
+			
+			// A root is a node with exec capability but no connected exec inputs from our set
+			if (bHasExecInput && !bHasConnectedExecInput)
+			{
+				bIsRoot = true;
+			}
+		}
+		
+		if (bIsRoot)
+		{
+			RootNodes.Add(Node);
+			NodeDepths.Add(Node, 0);
+		}
+	}
+	
+	// If no root nodes found, use all nodes and find the leftmost ones
+	if (RootNodes.Num() == 0)
+	{
+		// Sort by X position and use the leftmost as root
+		NodesToLayout.Sort([](const UEdGraphNode& A, const UEdGraphNode& B) {
+			return A.NodePosX < B.NodePosX;
+		});
+		
+		if (NodesToLayout.Num() > 0)
+		{
+			RootNodes.Add(NodesToLayout[0]);
+			NodeDepths.Add(NodesToLayout[0], 0);
+		}
+	}
+	
+	// BFS traversal to assign depths following execution flow
+	TQueue<UEdGraphNode*> TraversalQueue;
+	TSet<UEdGraphNode*> Visited;
+	
+	for (UEdGraphNode* Root : RootNodes)
+	{
+		TraversalQueue.Enqueue(Root);
+		Visited.Add(Root);
+	}
+	
+	while (!TraversalQueue.IsEmpty())
+	{
+		UEdGraphNode* Current;
+		TraversalQueue.Dequeue(Current);
+		
+		int32 CurrentDepth = NodeDepths.FindRef(Current);
+		
+		// Find all nodes connected via exec output pins
+		for (UEdGraphPin* Pin : Current->Pins)
+		{
+			if (Pin && Pin->Direction == EGPD_Output && 
+				Pin->PinType.PinCategory == UEdGraphSchema_K2::PC_Exec)
+			{
+				for (UEdGraphPin* LinkedPin : Pin->LinkedTo)
+				{
+					if (LinkedPin)
+					{
+						UEdGraphNode* NextNode = LinkedPin->GetOwningNode();
+						if (NextNode && NodesToLayout.Contains(NextNode) && !Visited.Contains(NextNode))
+						{
+							Visited.Add(NextNode);
+							NodeDepths.Add(NextNode, CurrentDepth + 1);
+							TraversalQueue.Enqueue(NextNode);
+						}
+					}
+				}
+			}
+		}
+	}
+	
+	// Handle nodes not reached via exec flow (pure nodes, data-only connections)
+	// Place them near their consumers
+	for (UEdGraphNode* Node : NodesToLayout)
+	{
+		if (!NodeDepths.Contains(Node))
+		{
+			// Find the minimum depth of connected nodes
+			int32 MinConsumerDepth = INT_MAX;
+			
+			for (UEdGraphPin* Pin : Node->Pins)
+			{
+				if (Pin && Pin->Direction == EGPD_Output)
+				{
+					for (UEdGraphPin* LinkedPin : Pin->LinkedTo)
+					{
+						if (LinkedPin)
+						{
+							UEdGraphNode* Consumer = LinkedPin->GetOwningNode();
+							if (Consumer && NodeDepths.Contains(Consumer))
+							{
+								MinConsumerDepth = FMath::Min(MinConsumerDepth, NodeDepths[Consumer]);
+							}
+						}
+					}
+				}
+			}
+			
+			// Place one level before the consumer, or at depth 0 if no consumers
+			if (MinConsumerDepth == INT_MAX)
+			{
+				NodeDepths.Add(Node, 0);
+			}
+			else
+			{
+				NodeDepths.Add(Node, FMath::Max(0, MinConsumerDepth - 1));
+			}
+		}
+	}
+	
+	// Group nodes by depth
+	TMap<int32, TArray<UEdGraphNode*>> NodesByDepth;
+	int32 MaxDepth = 0;
+	
+	for (const auto& Pair : NodeDepths)
+	{
+		NodesByDepth.FindOrAdd(Pair.Value).Add(Pair.Key);
+		MaxDepth = FMath::Max(MaxDepth, Pair.Value);
+	}
+	
+	// Calculate starting position (find minimum current position to use as anchor)
+	int32 StartX = INT_MAX;
+	int32 StartY = INT_MAX;
+	
+	for (UEdGraphNode* Node : NodesToLayout)
+	{
+		StartX = FMath::Min(StartX, Node->NodePosX);
+		StartY = FMath::Min(StartY, Node->NodePosY);
+	}
+	
+	// If no valid start position, use origin
+	if (StartX == INT_MAX) StartX = 0;
+	if (StartY == INT_MAX) StartY = 0;
+	
+	// Position nodes based on strategy
+	int32 NodesPositioned = 0;
+	
+	if (Strategy == TEXT("vertical"))
+	{
+		// Vertical: arrange top-to-bottom following execution flow
+		for (int32 Depth = 0; Depth <= MaxDepth; Depth++)
+		{
+			TArray<UEdGraphNode*>* NodesAtDepth = NodesByDepth.Find(Depth);
+			if (!NodesAtDepth)
+			{
+				continue;
+			}
+			
+			int32 Y = StartY + Depth * SpacingY;
+			int32 X = StartX;
+			
+			for (UEdGraphNode* Node : *NodesAtDepth)
+			{
+				Node->NodePosX = X;
+				Node->NodePosY = Y;
+				X += SpacingX;
+				NodesPositioned++;
+			}
+		}
+	}
+	else if (Strategy == TEXT("tree"))
+	{
+		// Tree: center children under parents
+		// First pass: calculate subtree widths
+		TMap<UEdGraphNode*, int32> SubtreeWidths;
+		
+		// Process from deepest to shallowest
+		for (int32 Depth = MaxDepth; Depth >= 0; Depth--)
+		{
+			TArray<UEdGraphNode*>* NodesAtDepth = NodesByDepth.Find(Depth);
+			if (!NodesAtDepth)
+			{
+				continue;
+			}
+			
+			for (UEdGraphNode* Node : *NodesAtDepth)
+			{
+				int32 Width = 1;
+				
+				// Sum widths of children
+				for (UEdGraphPin* Pin : Node->Pins)
+				{
+					if (Pin && Pin->Direction == EGPD_Output && 
+						Pin->PinType.PinCategory == UEdGraphSchema_K2::PC_Exec)
+					{
+						for (UEdGraphPin* LinkedPin : Pin->LinkedTo)
+						{
+							if (LinkedPin)
+							{
+								UEdGraphNode* Child = LinkedPin->GetOwningNode();
+								if (Child && SubtreeWidths.Contains(Child))
+								{
+									Width += SubtreeWidths[Child];
+								}
+							}
+						}
+					}
+				}
+				
+				SubtreeWidths.Add(Node, FMath::Max(1, Width));
+			}
+		}
+		
+		// Second pass: position nodes
+		TMap<UEdGraphNode*, int32> NodeXPositions;
+		int32 CurrentX = StartX;
+		
+		// Position roots first
+		for (UEdGraphNode* Root : RootNodes)
+		{
+			int32 Width = SubtreeWidths.FindRef(Root);
+			NodeXPositions.Add(Root, CurrentX + (Width * SpacingX) / 2);
+			Root->NodePosX = NodeXPositions[Root];
+			Root->NodePosY = StartY;
+			CurrentX += Width * SpacingX;
+			NodesPositioned++;
+		}
+		
+		// Position children relative to parents
+		for (int32 Depth = 1; Depth <= MaxDepth; Depth++)
+		{
+			TArray<UEdGraphNode*>* NodesAtDepth = NodesByDepth.Find(Depth);
+			if (!NodesAtDepth)
+			{
+				continue;
+			}
+			
+			for (UEdGraphNode* Node : *NodesAtDepth)
+			{
+				if (!NodeXPositions.Contains(Node))
+				{
+					// Find parent and position relative to it
+					Node->NodePosX = StartX;
+					Node->NodePosY = StartY + Depth * SpacingY;
+					NodeXPositions.Add(Node, Node->NodePosX);
+				}
+				else
+				{
+					Node->NodePosY = StartY + Depth * SpacingY;
+				}
+				NodesPositioned++;
+			}
+		}
+	}
+	else if (Strategy == TEXT("compact"))
+	{
+		// Compact: minimize total area
+		// Use a simple grid packing
+		int32 NodesPerRow = FMath::Max(1, FMath::CeilToInt(FMath::Sqrt((float)NodesToLayout.Num())));
+		int32 Index = 0;
+		
+		for (UEdGraphNode* Node : NodesToLayout)
+		{
+			int32 Row = Index / NodesPerRow;
+			int32 Col = Index % NodesPerRow;
+			
+			Node->NodePosX = StartX + Col * SpacingX;
+			Node->NodePosY = StartY + Row * SpacingY;
+			
+			Index++;
+			NodesPositioned++;
+		}
+	}
+	else // horizontal (default)
+	{
+		// Horizontal: arrange left-to-right following execution flow
+		for (int32 Depth = 0; Depth <= MaxDepth; Depth++)
+		{
+			TArray<UEdGraphNode*>* NodesAtDepth = NodesByDepth.Find(Depth);
+			if (!NodesAtDepth)
+			{
+				continue;
+			}
+			
+			// Sort nodes at this depth by their Y position to maintain relative order
+			NodesAtDepth->Sort([](const UEdGraphNode& A, const UEdGraphNode& B) {
+				return A.NodePosY < B.NodePosY;
+			});
+			
+			int32 X = StartX + Depth * SpacingX;
+			int32 Y = StartY;
+			
+			for (UEdGraphNode* Node : *NodesAtDepth)
+			{
+				Node->NodePosX = X;
+				Node->NodePosY = Y;
+				Y += SpacingY;
+				NodesPositioned++;
+			}
+		}
+	}
+	
+	// Mark blueprint as modified
+	FBlueprintEditorUtils::MarkBlueprintAsStructurallyModified(Blueprint);
+	
+	// Build result
+	TSharedPtr<FJsonObject> Result = MakeResult();
+	Result->SetStringField(TEXT("graph_name"), GraphName);
+	Result->SetStringField(TEXT("strategy"), Strategy);
+	Result->SetNumberField(TEXT("nodes_positioned"), NodesPositioned);
+	Result->SetNumberField(TEXT("root_nodes"), RootNodes.Num());
+	Result->SetNumberField(TEXT("max_depth"), MaxDepth);
+	Result->SetNumberField(TEXT("spacing_x"), SpacingX);
+	Result->SetNumberField(TEXT("spacing_y"), SpacingY);
+	
+	// Include positioned node info
+	TArray<TSharedPtr<FJsonValue>> PositionedNodesArray;
+	for (UEdGraphNode* Node : NodesToLayout)
+	{
+		TSharedPtr<FJsonObject> NodeInfo = MakeShared<FJsonObject>();
+		NodeInfo->SetStringField(TEXT("node_id"), Node->NodeGuid.ToString());
+		NodeInfo->SetStringField(TEXT("node_title"), Node->GetNodeTitle(ENodeTitleType::FullTitle).ToString());
+		NodeInfo->SetNumberField(TEXT("x"), Node->NodePosX);
+		NodeInfo->SetNumberField(TEXT("y"), Node->NodePosY);
+		NodeInfo->SetNumberField(TEXT("depth"), NodeDepths.FindRef(Node));
+		PositionedNodesArray.Add(MakeShared<FJsonValueObject>(NodeInfo));
+	}
+	Result->SetArrayField(TEXT("positioned_nodes"), PositionedNodesArray);
+	
+	return FECACommandResult::Success(Result);
+}
+#endif // End of legacy AutoLayoutBlueprintGraph
