ECABridge experimental MCP plugin

diff --git a/Engine/Plugins/Experimental/ECABridge/Source/ECABridge/Private/Commands/ECACommand.cpp b/Engine/Plugins/Experimental/ECABridge/Source/ECABridge/Private/Commands/ECACommand.cpp
new file mode 100644
--- /dev/null
+++ b/Engine/Plugins/Experimental/ECABridge/Source/ECABridge/Private/Commands/ECACommand.cpp
@@ -0,0 +1,416 @@
+// Copyright Epic Games, Inc. All Rights Reserved.
+
+#include "Commands/ECACommand.h"
+#include "Editor.h"
+#include "Engine/World.h"
+#include "EngineUtils.h"
+#include "Engine/Blueprint.h"
+#include "GameFramework/Actor.h"
+#include "JsonObjectConverter.h"
+
+//------------------------------------------------------------------------------
+// FECACommandResult
+//------------------------------------------------------------------------------
+
+FString FECACommandResult::ToJsonString() const
+{
+	TSharedPtr<FJsonObject> Response = MakeShared<FJsonObject>();
+	Response->SetBoolField(TEXT("success"), bSuccess);
+	
+	if (bSuccess)
+	{
+		if (ResultData.IsValid())
+		{
+			Response->SetObjectField(TEXT("result"), ResultData);
+		}
+	}
+	else
+	{
+		Response->SetStringField(TEXT("error"), ErrorMessage);
+	}
+	
+	FString OutputString;
+	TSharedRef<TJsonWriter<>> Writer = TJsonWriterFactory<>::Create(&OutputString);
+	FJsonSerializer::Serialize(Response.ToSharedRef(), Writer);
+	return OutputString;
+}
+
+//------------------------------------------------------------------------------
+// IECACommand - Parameter Helpers
+//------------------------------------------------------------------------------
+
+bool IECACommand::GetStringParam(const TSharedPtr<FJsonObject>& Params, const FString& Name, FString& OutValue, bool bRequired)
+{
+	if (!Params.IsValid())
+	{
+		return false;
+	}
+	
+	if (Params->TryGetStringField(Name, OutValue))
+	{
+		return true;
+	}
+	
+	return false;
+}
+
+bool IECACommand::GetIntParam(const TSharedPtr<FJsonObject>& Params, const FString& Name, int32& OutValue, bool bRequired)
+{
+	if (!Params.IsValid())
+	{
+		return false;
+	}
+	
+	double Value;
+	if (Params->TryGetNumberField(Name, Value))
+	{
+		OutValue = FMath::RoundToInt(Value);
+		return true;
+	}
+	
+	return false;
+}
+
+bool IECACommand::GetFloatParam(const TSharedPtr<FJsonObject>& Params, const FString& Name, double& OutValue, bool bRequired)
+{
+	if (!Params.IsValid())
+	{
+		return false;
+	}
+	
+	return Params->TryGetNumberField(Name, OutValue);
+}
+
+bool IECACommand::GetBoolParam(const TSharedPtr<FJsonObject>& Params, const FString& Name, bool& OutValue, bool bRequired)
+{
+	if (!Params.IsValid())
+	{
+		return false;
+	}
+	
+	return Params->TryGetBoolField(Name, OutValue);
+}
+
+bool IECACommand::GetVectorParam(const TSharedPtr<FJsonObject>& Params, const FString& Name, FVector& OutValue, bool bRequired)
+{
+	if (!Params.IsValid())
+	{
+		return false;
+	}
+	
+	// Try object format first: {"x": 1, "y": 2, "z": 3}
+	const TSharedPtr<FJsonObject>* VecObj;
+	if (Params->TryGetObjectField(Name, VecObj))
+	{
+		OutValue.X = (*VecObj)->GetNumberField(TEXT("x"));
+		OutValue.Y = (*VecObj)->GetNumberField(TEXT("y"));
+		OutValue.Z = (*VecObj)->GetNumberField(TEXT("z"));
+		return true;
+	}
+	
+	// Try array format: [1, 2, 3]
+	const TArray<TSharedPtr<FJsonValue>>* ArrValue;
+	if (Params->TryGetArrayField(Name, ArrValue) && ArrValue->Num() >= 3)
+	{
+		OutValue.X = (*ArrValue)[0]->AsNumber();
+		OutValue.Y = (*ArrValue)[1]->AsNumber();
+		OutValue.Z = (*ArrValue)[2]->AsNumber();
+		return true;
+	}
+	
+	// Try string format: "1,2,3" or "1, 2, 3"
+	FString StrValue;
+	if (Params->TryGetStringField(Name, StrValue))
+	{
+		TArray<FString> Parts;
+		StrValue.ParseIntoArray(Parts, TEXT(","));
+		if (Parts.Num() >= 3)
+		{
+			OutValue.X = FCString::Atof(*Parts[0].TrimStartAndEnd());
+			OutValue.Y = FCString::Atof(*Parts[1].TrimStartAndEnd());
+			OutValue.Z = FCString::Atof(*Parts[2].TrimStartAndEnd());
+			return true;
+		}
+	}
+	
+	return false;
+}
+
+bool IECACommand::GetRotatorParam(const TSharedPtr<FJsonObject>& Params, const FString& Name, FRotator& OutValue, bool bRequired)
+{
+	if (!Params.IsValid())
+	{
+		return false;
+	}
+	
+	// Try object format first: {"pitch": 1, "yaw": 2, "roll": 3}
+	const TSharedPtr<FJsonObject>* RotObj;
+	if (Params->TryGetObjectField(Name, RotObj))
+	{
+		OutValue.Pitch = (*RotObj)->GetNumberField(TEXT("pitch"));
+		OutValue.Yaw = (*RotObj)->GetNumberField(TEXT("yaw"));
+		OutValue.Roll = (*RotObj)->GetNumberField(TEXT("roll"));
+		return true;
+	}
+	
+	// Try array format: [pitch, yaw, roll]
+	const TArray<TSharedPtr<FJsonValue>>* ArrValue;
+	if (Params->TryGetArrayField(Name, ArrValue) && ArrValue->Num() >= 3)
+	{
+		OutValue.Pitch = (*ArrValue)[0]->AsNumber();
+		OutValue.Yaw = (*ArrValue)[1]->AsNumber();
+		OutValue.Roll = (*ArrValue)[2]->AsNumber();
+		return true;
+	}
+	
+	// Try string format: "pitch,yaw,roll" or "pitch, yaw, roll"
+	FString StrValue;
+	if (Params->TryGetStringField(Name, StrValue))
+	{
+		TArray<FString> Parts;
+		StrValue.ParseIntoArray(Parts, TEXT(","));
+		if (Parts.Num() >= 3)
+		{
+			OutValue.Pitch = FCString::Atof(*Parts[0].TrimStartAndEnd());
+			OutValue.Yaw = FCString::Atof(*Parts[1].TrimStartAndEnd());
+			OutValue.Roll = FCString::Atof(*Parts[2].TrimStartAndEnd());
+			return true;
+		}
+	}
+	
+	return false;
+}
+
+bool IECACommand::GetArrayParam(const TSharedPtr<FJsonObject>& Params, const FString& Name, const TArray<TSharedPtr<FJsonValue>>*& OutValue, bool bRequired)
+{
+	if (!Params.IsValid())
+	{
+		return false;
+	}
+	
+	return Params->TryGetArrayField(Name, OutValue);
+}
+
+bool IECACommand::GetObjectParam(const TSharedPtr<FJsonObject>& Params, const FString& Name, const TSharedPtr<FJsonObject>*& OutValue, bool bRequired)
+{
+	if (!Params.IsValid())
+	{
+		return false;
+	}
+	
+	return Params->TryGetObjectField(Name, OutValue);
+}
+
+//------------------------------------------------------------------------------
+// IECACommand - Result Helpers
+//------------------------------------------------------------------------------
+
+TSharedPtr<FJsonObject> IECACommand::MakeResult()
+{
+	return MakeShared<FJsonObject>();
+}
+
+TSharedPtr<FJsonObject> IECACommand::VectorToJson(const FVector& Vec)
+{
+	TSharedPtr<FJsonObject> Obj = MakeShared<FJsonObject>();
+	Obj->SetNumberField(TEXT("x"), Vec.X);
+	Obj->SetNumberField(TEXT("y"), Vec.Y);
+	Obj->SetNumberField(TEXT("z"), Vec.Z);
+	return Obj;
+}
+
+TSharedPtr<FJsonObject> IECACommand::RotatorToJson(const FRotator& Rot)
+{
+	TSharedPtr<FJsonObject> Obj = MakeShared<FJsonObject>();
+	Obj->SetNumberField(TEXT("pitch"), Rot.Pitch);
+	Obj->SetNumberField(TEXT("yaw"), Rot.Yaw);
+	Obj->SetNumberField(TEXT("roll"), Rot.Roll);
+	return Obj;
+}
+
+TSharedPtr<FJsonObject> IECACommand::TransformToJson(const FTransform& Transform)
+{
+	TSharedPtr<FJsonObject> Obj = MakeShared<FJsonObject>();
+	Obj->SetObjectField(TEXT("location"), VectorToJson(Transform.GetLocation()));
+	Obj->SetObjectField(TEXT("rotation"), RotatorToJson(Transform.Rotator()));
+	Obj->SetObjectField(TEXT("scale"), VectorToJson(Transform.GetScale3D()));
+	return Obj;
+}
+
+//------------------------------------------------------------------------------
+// IECACommand - Common Operations
+//------------------------------------------------------------------------------
+
+UWorld* IECACommand::GetEditorWorld()
+{
+	if (GEditor)
+	{
+		return GEditor->GetEditorWorldContext().World();
+	}
+	return nullptr;
+}
+
+AActor* IECACommand::FindActorByName(const FString& ActorName)
+{
+	UWorld* World = GetEditorWorld();
+	if (!World)
+	{
+		return nullptr;
+	}
+	
+	// Check if the name contains wildcards
+	bool bHasWildcard = ActorName.Contains(TEXT("*")) || ActorName.Contains(TEXT("?"));
+	
+	for (TActorIterator<AActor> It(World); It; ++It)
+	{
+		AActor* Actor = *It;
+		FString ActorLabel = Actor->GetActorLabel();
+		FString ActorObjName = Actor->GetName();
+		
+		if (bHasWildcard)
+		{
+			// Wildcard matching (case-insensitive)
+			if (ActorLabel.MatchesWildcard(ActorName, ESearchCase::IgnoreCase) ||
+			    ActorObjName.MatchesWildcard(ActorName, ESearchCase::IgnoreCase))
+			{
+				return Actor;
+			}
+		}
+		else
+		{
+			// Exact matching (case-insensitive)
+			if (ActorLabel.Equals(ActorName, ESearchCase::IgnoreCase) ||
+			    ActorObjName.Equals(ActorName, ESearchCase::IgnoreCase))
+			{
+				return Actor;
+			}
+		}
+	}
+	
+	return nullptr;
+}
+
+UBlueprint* IECACommand::LoadBlueprintByPath(const FString& BlueprintPath)
+{
+	return LoadObject<UBlueprint>(nullptr, *BlueprintPath);
+}
+
+//------------------------------------------------------------------------------
+// FECACommandRegistry
+//------------------------------------------------------------------------------
+
+FECACommandRegistry& FECACommandRegistry::Get()
+{
+	static FECACommandRegistry Instance;
+	return Instance;
+}
+
+void FECACommandRegistry::RegisterCommand(TSharedPtr<IECACommand> Command)
+{
+	if (!Command.IsValid())
+	{
+		return;
+	}
+	
+	FScopeLock Lock(&CommandsLock);
+	
+	const FString Name = Command->GetName();
+	if (Commands.Contains(Name))
+	{
+		UE_LOG(LogTemp, Warning, TEXT("[ECABridge] Command '%s' already registered, replacing"), *Name);
+	}
+	
+	Commands.Add(Name, Command);
+	UE_LOG(LogTemp, Log, TEXT("[ECABridge] Registered command: %s (%s)"), *Name, *Command->GetCategory());
+}
+
+void FECACommandRegistry::UnregisterCommand(const FString& Name)
+{
+	FScopeLock Lock(&CommandsLock);
+	Commands.Remove(Name);
+}
+
+TSharedPtr<IECACommand> FECACommandRegistry::GetCommand(const FString& Name) const
+{
+	FScopeLock Lock(&CommandsLock);
+	const TSharedPtr<IECACommand>* Found = Commands.Find(Name);
+	return Found ? *Found : nullptr;
+}
+
+TArray<TSharedPtr<IECACommand>> FECACommandRegistry::GetAllCommands() const
+{
+	FScopeLock Lock(&CommandsLock);
+	TArray<TSharedPtr<IECACommand>> Result;
+	Commands.GenerateValueArray(Result);
+	return Result;
+}
+
+TArray<TSharedPtr<IECACommand>> FECACommandRegistry::GetCommandsByCategory(const FString& Category) const
+{
+	FScopeLock Lock(&CommandsLock);
+	TArray<TSharedPtr<IECACommand>> Result;
+	
+	for (const auto& Pair : Commands)
+	{
+		if (Pair.Value->GetCategory() == Category)
+		{
+			Result.Add(Pair.Value);
+		}
+	}
+	
+	return Result;
+}
+
+bool FECACommandRegistry::HasCommand(const FString& Name) const
+{
+	FScopeLock Lock(&CommandsLock);
+	return Commands.Contains(Name);
+}
+
+TArray<FString> FECACommandRegistry::GetCategories() const
+{
+	FScopeLock Lock(&CommandsLock);
+	TSet<FString> Categories;
+	
+	for (const auto& Pair : Commands)
+	{
+		Categories.Add(Pair.Value->GetCategory());
+	}
+	
+	return Categories.Array();
+}
+
+FECACommandResult FECACommandRegistry::ExecuteCommand(const FString& Name, const TSharedPtr<FJsonObject>& Params)
+{
+	TSharedPtr<IECACommand> Command = GetCommand(Name);
+	
+	if (!Command.IsValid())
+	{
+		return FECACommandResult::Error(FString::Printf(TEXT("Unknown command: %s"), *Name));
+	}
+	
+	return Command->Execute(Params);
+}
+
+void FECACommandRegistry::LogCommands() const
+{
+	FScopeLock Lock(&CommandsLock);
+	
+	UE_LOG(LogTemp, Log, TEXT("[ECABridge] Registered commands (%d total):"), Commands.Num());
+	
+	TArray<FString> Categories = GetCategories();
+	Categories.Sort();
+	
+	for (const FString& Category : Categories)
+	{
+		UE_LOG(LogTemp, Log, TEXT("  [%s]"), *Category);
+		
+		for (const auto& Pair : Commands)
+		{
+			if (Pair.Value->GetCategory() == Category)
+			{
+				UE_LOG(LogTemp, Log, TEXT("    - %s: %s"), *Pair.Key, *Pair.Value->GetDescription());
+			}
+		}
+	}
+}
diff --git a/Engine/Plugins/Experimental/ECABridge/Source/ECABridge/Private/Commands/ECAComponentCommands.cpp b/Engine/Plugins/Experimental/ECABridge/Source/ECABridge/Private/Commands/ECAComponentCommands.cpp
new file mode 100644
--- /dev/null
+++ b/Engine/Plugins/Experimental/ECABridge/Source/ECABridge/Private/Commands/ECAComponentCommands.cpp
@@ -0,0 +1,385 @@
+// Copyright Epic Games, Inc. All Rights Reserved.
+
+#include "Commands/ECAComponentCommands.h"
+#include "Commands/ECACommand.h"
+#include "Editor.h"
+#include "Engine/Blueprint.h"
+#include "Engine/SimpleConstructionScript.h"
+#include "Engine/SCS_Node.h"
+#include "Kismet2/BlueprintEditorUtils.h"
+#include "Components/StaticMeshComponent.h"
+#include "Components/PrimitiveComponent.h"
+#include "Engine/StaticMesh.h"
+#include "PhysicsEngine/BodyInstance.h"
+
+// Register all component commands
+REGISTER_ECA_COMMAND(FECACommand_SetStaticMeshProperties)
+REGISTER_ECA_COMMAND(FECACommand_SetPhysicsProperties)
+REGISTER_ECA_COMMAND(FECACommand_SetComponentProperty)
+REGISTER_ECA_COMMAND(FECACommand_SetComponentTransform)
+REGISTER_ECA_COMMAND(FECACommand_GetBlueprintComponents)
+
+//------------------------------------------------------------------------------
+// Helper: Find SCS Node by name
+//------------------------------------------------------------------------------
+
+static USCS_Node* FindSCSNodeByName(UBlueprint* Blueprint, const FString& ComponentName)
+{
+	if (!Blueprint || !Blueprint->SimpleConstructionScript)
+	{
+		return nullptr;
+	}
+	
+	TArray<USCS_Node*> AllNodes = Blueprint->SimpleConstructionScript->GetAllNodes();
+	for (USCS_Node* Node : AllNodes)
+	{
+		if (Node && Node->GetVariableName().ToString() == ComponentName)
+		{
+			return Node;
+		}
+	}
+	
+	return nullptr;
+}
+
+//------------------------------------------------------------------------------
+// SetStaticMeshProperties
+//------------------------------------------------------------------------------
+
+FECACommandResult FECACommand_SetStaticMeshProperties::Execute(const TSharedPtr<FJsonObject>& Params)
+{
+	FString BlueprintPath;
+	if (!GetStringParam(Params, TEXT("blueprint_path"), BlueprintPath))
+	{
+		return FECACommandResult::Error(TEXT("Missing required parameter: blueprint_path"));
+	}
+	
+	FString ComponentName;
+	if (!GetStringParam(Params, TEXT("component_name"), ComponentName))
+	{
+		return FECACommandResult::Error(TEXT("Missing required parameter: component_name"));
+	}
+	
+	FString StaticMeshPath;
+	if (!GetStringParam(Params, TEXT("static_mesh"), StaticMeshPath))
+	{
+		return FECACommandResult::Error(TEXT("Missing required parameter: static_mesh"));
+	}
+	
+	UBlueprint* Blueprint = LoadBlueprintByPath(BlueprintPath);
+	if (!Blueprint)
+	{
+		return FECACommandResult::Error(FString::Printf(TEXT("Blueprint not found: %s"), *BlueprintPath));
+	}
+	
+	USCS_Node* Node = FindSCSNodeByName(Blueprint, ComponentName);
+	if (!Node)
+	{
+		return FECACommandResult::Error(FString::Printf(TEXT("Component not found: %s"), *ComponentName));
+	}
+	
+	// Get the component template
+	UStaticMeshComponent* MeshComponent = Cast<UStaticMeshComponent>(Node->ComponentTemplate);
+	if (!MeshComponent)
+	{
+		return FECACommandResult::Error(FString::Printf(TEXT("Component is not a StaticMeshComponent: %s"), *ComponentName));
+	}
+	
+	// Load the mesh
+	UStaticMesh* Mesh = LoadObject<UStaticMesh>(nullptr, *StaticMeshPath);
+	if (!Mesh)
+	{
+		return FECACommandResult::Error(FString::Printf(TEXT("Static mesh not found: %s"), *StaticMeshPath));
+	}
+	
+	// Set the mesh
+	MeshComponent->SetStaticMesh(Mesh);
+	
+	FBlueprintEditorUtils::MarkBlueprintAsStructurallyModified(Blueprint);
+	
+	TSharedPtr<FJsonObject> Result = MakeResult();
+	Result->SetStringField(TEXT("component_name"), ComponentName);
+	Result->SetStringField(TEXT("static_mesh"), StaticMeshPath);
+	return FECACommandResult::Success(Result);
+}
+
+//------------------------------------------------------------------------------
+// SetPhysicsProperties
+//------------------------------------------------------------------------------
+
+FECACommandResult FECACommand_SetPhysicsProperties::Execute(const TSharedPtr<FJsonObject>& Params)
+{
+	FString BlueprintPath;
+	if (!GetStringParam(Params, TEXT("blueprint_path"), BlueprintPath))
+	{
+		return FECACommandResult::Error(TEXT("Missing required parameter: blueprint_path"));
+	}
+	
+	FString ComponentName;
+	if (!GetStringParam(Params, TEXT("component_name"), ComponentName))
+	{
+		return FECACommandResult::Error(TEXT("Missing required parameter: component_name"));
+	}
+	
+	bool bSimulatePhysics = true;
+	bool bGravityEnabled = true;
+	double Mass = 1.0;
+	double LinearDamping = 0.01;
+	double AngularDamping = 0.0;
+	
+	GetBoolParam(Params, TEXT("simulate_physics"), bSimulatePhysics, false);
+	GetBoolParam(Params, TEXT("gravity_enabled"), bGravityEnabled, false);
+	GetFloatParam(Params, TEXT("mass"), Mass, false);
+	GetFloatParam(Params, TEXT("linear_damping"), LinearDamping, false);
+	GetFloatParam(Params, TEXT("angular_damping"), AngularDamping, false);
+	
+	UBlueprint* Blueprint = LoadBlueprintByPath(BlueprintPath);
+	if (!Blueprint)
+	{
+		return FECACommandResult::Error(FString::Printf(TEXT("Blueprint not found: %s"), *BlueprintPath));
+	}
+	
+	USCS_Node* Node = FindSCSNodeByName(Blueprint, ComponentName);
+	if (!Node)
+	{
+		return FECACommandResult::Error(FString::Printf(TEXT("Component not found: %s"), *ComponentName));
+	}
+	
+	// Get the component template
+	UPrimitiveComponent* PrimitiveComponent = Cast<UPrimitiveComponent>(Node->ComponentTemplate);
+	if (!PrimitiveComponent)
+	{
+		return FECACommandResult::Error(FString::Printf(TEXT("Component is not a PrimitiveComponent: %s"), *ComponentName));
+	}
+	
+	// Set physics properties
+	PrimitiveComponent->SetSimulatePhysics(bSimulatePhysics);
+	PrimitiveComponent->SetEnableGravity(bGravityEnabled);
+	PrimitiveComponent->SetMassOverrideInKg(NAME_None, Mass);
+	PrimitiveComponent->SetLinearDamping(LinearDamping);
+	PrimitiveComponent->SetAngularDamping(AngularDamping);
+	
+	FBlueprintEditorUtils::MarkBlueprintAsStructurallyModified(Blueprint);
+	
+	TSharedPtr<FJsonObject> Result = MakeResult();
+	Result->SetStringField(TEXT("component_name"), ComponentName);
+	Result->SetBoolField(TEXT("simulate_physics"), bSimulatePhysics);
+	Result->SetBoolField(TEXT("gravity_enabled"), bGravityEnabled);
+	Result->SetNumberField(TEXT("mass"), Mass);
+	Result->SetNumberField(TEXT("linear_damping"), LinearDamping);
+	Result->SetNumberField(TEXT("angular_damping"), AngularDamping);
+	return FECACommandResult::Success(Result);
+}
+
+//------------------------------------------------------------------------------
+// SetComponentProperty
+//------------------------------------------------------------------------------
+
+FECACommandResult FECACommand_SetComponentProperty::Execute(const TSharedPtr<FJsonObject>& Params)
+{
+	FString BlueprintPath;
+	if (!GetStringParam(Params, TEXT("blueprint_path"), BlueprintPath))
+	{
+		return FECACommandResult::Error(TEXT("Missing required parameter: blueprint_path"));
+	}
+	
+	FString ComponentName;
+	if (!GetStringParam(Params, TEXT("component_name"), ComponentName))
+	{
+		return FECACommandResult::Error(TEXT("Missing required parameter: component_name"));
+	}
+	
+	FString PropertyName;
+	if (!GetStringParam(Params, TEXT("property_name"), PropertyName))
+	{
+		return FECACommandResult::Error(TEXT("Missing required parameter: property_name"));
+	}
+	
+	UBlueprint* Blueprint = LoadBlueprintByPath(BlueprintPath);
+	if (!Blueprint)
+	{
+		return FECACommandResult::Error(FString::Printf(TEXT("Blueprint not found: %s"), *BlueprintPath));
+	}
+	
+	USCS_Node* Node = FindSCSNodeByName(Blueprint, ComponentName);
+	if (!Node || !Node->ComponentTemplate)
+	{
+		return FECACommandResult::Error(FString::Printf(TEXT("Component not found: %s"), *ComponentName));
+	}
+	
+	UActorComponent* Component = Node->ComponentTemplate;
+	
+	// Find the property
+	FProperty* Property = Component->GetClass()->FindPropertyByName(*PropertyName);
+	if (!Property)
+	{
+		return FECACommandResult::Error(FString::Printf(TEXT("Property not found: %s"), *PropertyName));
+	}
+	
+	// Get the property value from JSON
+	TSharedPtr<FJsonValue> PropertyValue = Params->TryGetField(TEXT("property_value"));
+	if (!PropertyValue.IsValid())
+	{
+		return FECACommandResult::Error(TEXT("Missing required parameter: property_value"));
+	}
+	
+	// Set property based on type
+	void* PropertyAddr = Property->ContainerPtrToValuePtr<void>(Component);
+	
+	if (FBoolProperty* BoolProp = CastField<FBoolProperty>(Property))
+	{
+		BoolProp->SetPropertyValue(PropertyAddr, PropertyValue->AsBool());
+	}
+	else if (FIntProperty* IntProp = CastField<FIntProperty>(Property))
+	{
+		IntProp->SetPropertyValue(PropertyAddr, FMath::RoundToInt(PropertyValue->AsNumber()));
+	}
+	else if (FFloatProperty* FloatProp = CastField<FFloatProperty>(Property))
+	{
+		FloatProp->SetPropertyValue(PropertyAddr, PropertyValue->AsNumber());
+	}
+	else if (FDoubleProperty* DoubleProp = CastField<FDoubleProperty>(Property))
+	{
+		DoubleProp->SetPropertyValue(PropertyAddr, PropertyValue->AsNumber());
+	}
+	else if (FStrProperty* StrProp = CastField<FStrProperty>(Property))
+	{
+		StrProp->SetPropertyValue(PropertyAddr, PropertyValue->AsString());
+	}
+	else if (FNameProperty* NameProp = CastField<FNameProperty>(Property))
+	{
+		NameProp->SetPropertyValue(PropertyAddr, FName(*PropertyValue->AsString()));
+	}
+	else
+	{
+		return FECACommandResult::Error(FString::Printf(TEXT("Unsupported property type for: %s"), *PropertyName));
+	}
+	
+	FBlueprintEditorUtils::MarkBlueprintAsStructurallyModified(Blueprint);
+	
+	TSharedPtr<FJsonObject> Result = MakeResult();
+	Result->SetStringField(TEXT("component_name"), ComponentName);
+	Result->SetStringField(TEXT("property_name"), PropertyName);
+	return FECACommandResult::Success(Result);
+}
+
+//------------------------------------------------------------------------------
+// SetComponentTransform
+//------------------------------------------------------------------------------
+
+FECACommandResult FECACommand_SetComponentTransform::Execute(const TSharedPtr<FJsonObject>& Params)
+{
+	FString BlueprintPath;
+	if (!GetStringParam(Params, TEXT("blueprint_path"), BlueprintPath))
+	{
+		return FECACommandResult::Error(TEXT("Missing required parameter: blueprint_path"));
+	}
+	
+	FString ComponentName;
+	if (!GetStringParam(Params, TEXT("component_name"), ComponentName))
+	{
+		return FECACommandResult::Error(TEXT("Missing required parameter: component_name"));
+	}
+	
+	UBlueprint* Blueprint = LoadBlueprintByPath(BlueprintPath);
+	if (!Blueprint)
+	{
+		return FECACommandResult::Error(FString::Printf(TEXT("Blueprint not found: %s"), *BlueprintPath));
+	}
+	
+	USCS_Node* Node = FindSCSNodeByName(Blueprint, ComponentName);
+	if (!Node || !Node->ComponentTemplate)
+	{
+		return FECACommandResult::Error(FString::Printf(TEXT("Component not found: %s"), *ComponentName));
+	}
+	
+	USceneComponent* SceneComponent = Cast<USceneComponent>(Node->ComponentTemplate);
+	if (!SceneComponent)
+	{
+		return FECACommandResult::Error(FString::Printf(TEXT("Component is not a SceneComponent: %s"), *ComponentName));
+	}
+	
+	FVector Location;
+	FRotator Rotation;
+	FVector Scale;
+	
+	if (GetVectorParam(Params, TEXT("location"), Location, false))
+	{
+		SceneComponent->SetRelativeLocation(Location);
+	}
+	
+	if (GetRotatorParam(Params, TEXT("rotation"), Rotation, false))
+	{
+		SceneComponent->SetRelativeRotation(Rotation);
+	}
+	
+	if (GetVectorParam(Params, TEXT("scale"), Scale, false))
+	{
+		SceneComponent->SetRelativeScale3D(Scale);
+	}
+	
+	FBlueprintEditorUtils::MarkBlueprintAsStructurallyModified(Blueprint);
+	
+	TSharedPtr<FJsonObject> Result = MakeResult();
+	Result->SetStringField(TEXT("component_name"), ComponentName);
+	Result->SetObjectField(TEXT("transform"), TransformToJson(SceneComponent->GetRelativeTransform()));
+	return FECACommandResult::Success(Result);
+}
+
+//------------------------------------------------------------------------------
+// GetBlueprintComponents
+//------------------------------------------------------------------------------
+
+FECACommandResult FECACommand_GetBlueprintComponents::Execute(const TSharedPtr<FJsonObject>& Params)
+{
+	FString BlueprintPath;
+	if (!GetStringParam(Params, TEXT("blueprint_path"), BlueprintPath))
+	{
+		return FECACommandResult::Error(TEXT("Missing required parameter: blueprint_path"));
+	}
+	
+	UBlueprint* Blueprint = LoadBlueprintByPath(BlueprintPath);
+	if (!Blueprint)
+	{
+		return FECACommandResult::Error(FString::Printf(TEXT("Blueprint not found: %s"), *BlueprintPath));
+	}
+	
+	TArray<TSharedPtr<FJsonValue>> ComponentsArray;
+	
+	if (Blueprint->SimpleConstructionScript)
+	{
+		TArray<USCS_Node*> AllNodes = Blueprint->SimpleConstructionScript->GetAllNodes();
+		
+		for (USCS_Node* Node : AllNodes)
+		{
+			if (!Node || !Node->ComponentTemplate)
+			{
+				continue;
+			}
+			
+			TSharedPtr<FJsonObject> CompJson = MakeShared<FJsonObject>();
+			CompJson->SetStringField(TEXT("name"), Node->GetVariableName().ToString());
+			CompJson->SetStringField(TEXT("class"), Node->ComponentClass ? Node->ComponentClass->GetName() : TEXT("None"));
+			
+			// Check if it's a scene component and get transform
+			USceneComponent* SceneComp = Cast<USceneComponent>(Node->ComponentTemplate);
+			if (SceneComp)
+			{
+				CompJson->SetObjectField(TEXT("relative_transform"), TransformToJson(SceneComp->GetRelativeTransform()));
+			}
+			
+			// Get parent node name if any
+			if (Node->ParentComponentOrVariableName != NAME_None)
+			{
+				CompJson->SetStringField(TEXT("parent"), Node->ParentComponentOrVariableName.ToString());
+			}
+			
+			ComponentsArray.Add(MakeShared<FJsonValueObject>(CompJson));
+		}
+	}
+	
+	TSharedPtr<FJsonObject> Result = MakeResult();
+	Result->SetArrayField(TEXT("components"), ComponentsArray);
+	Result->SetNumberField(TEXT("count"), ComponentsArray.Num());
+	return FECACommandResult::Success(Result);
+}
diff --git a/Engine/Plugins/Experimental/ECABridge/Source/ECABridge/Private/Commands/ECADataTableCommands.cpp b/Engine/Plugins/Experimental/ECABridge/Source/ECABridge/Private/Commands/ECADataTableCommands.cpp
new file mode 100644
--- /dev/null
+++ b/Engine/Plugins/Experimental/ECABridge/Source/ECABridge/Private/Commands/ECADataTableCommands.cpp
@@ -0,0 +1,695 @@
+// Copyright Epic Games, Inc. All Rights Reserved.
+
+#include "Commands/ECADataTableCommands.h"
+#include "Commands/ECACommand.h"
+#include "Engine/DataTable.h"
+#include "UObject/UnrealType.h"
+#include "UObject/TextProperty.h"
+#include "UObject/EnumProperty.h"
+#include "Serialization/JsonSerializer.h"
+#include "Serialization/JsonWriter.h"
+#include "Dom/JsonObject.h"
+#include "Dom/JsonValue.h"
+#include "DataTableEditorUtils.h"
+
+// Register all DataTable commands
+REGISTER_ECA_COMMAND(FECACommand_GetDataTableSchema)
+REGISTER_ECA_COMMAND(FECACommand_GetDataTableRows)
+REGISTER_ECA_COMMAND(FECACommand_SetDataTableRow)
+REGISTER_ECA_COMMAND(FECACommand_DeleteDataTableRow)
+
+//------------------------------------------------------------------------------
+// Helper: Convert a single property value to a JSON value
+//------------------------------------------------------------------------------
+static TSharedPtr<FJsonValue> DataTablePropertyToJsonValue(FProperty* Property, const void* ValuePtr)
+{
+	if (!Property || !ValuePtr)
+	{
+		return MakeShared<FJsonValueNull>();
+	}
+
+	if (FBoolProperty* BoolProp = CastField<FBoolProperty>(Property))
+	{
+		return MakeShared<FJsonValueBoolean>(BoolProp->GetPropertyValue(ValuePtr));
+	}
+	else if (FNumericProperty* NumProp = CastField<FNumericProperty>(Property))
+	{
+		if (NumProp->IsFloatingPoint())
+		{
+			return MakeShared<FJsonValueNumber>(NumProp->GetFloatingPointPropertyValue(ValuePtr));
+		}
+		else if (NumProp->IsInteger())
+		{
+			return MakeShared<FJsonValueNumber>(static_cast<double>(NumProp->GetSignedIntPropertyValue(ValuePtr)));
+		}
+	}
+	else if (FStrProperty* StrProp = CastField<FStrProperty>(Property))
+	{
+		return MakeShared<FJsonValueString>(StrProp->GetPropertyValue(ValuePtr));
+	}
+	else if (FNameProperty* NameProp = CastField<FNameProperty>(Property))
+	{
+		return MakeShared<FJsonValueString>(NameProp->GetPropertyValue(ValuePtr).ToString());
+	}
+	else if (FTextProperty* TextProp = CastField<FTextProperty>(Property))
+	{
+		return MakeShared<FJsonValueString>(TextProp->GetPropertyValue(ValuePtr).ToString());
+	}
+	else if (FEnumProperty* EnumProp = CastField<FEnumProperty>(Property))
+	{
+		UEnum* Enum = EnumProp->GetEnum();
+		FNumericProperty* UnderlyingProp = EnumProp->GetUnderlyingProperty();
+		int64 EnumValue = UnderlyingProp->GetSignedIntPropertyValue(ValuePtr);
+		FString EnumName = Enum->GetNameStringByValue(EnumValue);
+		return MakeShared<FJsonValueString>(EnumName);
+	}
+	else if (FByteProperty* ByteProp = CastField<FByteProperty>(Property))
+	{
+		if (UEnum* Enum = ByteProp->GetIntPropertyEnum())
+		{
+			uint8 ByteValue = *static_cast<const uint8*>(ValuePtr);
+			FString EnumName = Enum->GetNameStringByValue(ByteValue);
+			return MakeShared<FJsonValueString>(EnumName);
+		}
+		else
+		{
+			return MakeShared<FJsonValueNumber>(static_cast<double>(*static_cast<const uint8*>(ValuePtr)));
+		}
+	}
+	else if (FStructProperty* StructProp = CastField<FStructProperty>(Property))
+	{
+		// For structs, recurse into their properties
+		TSharedPtr<FJsonObject> StructObj = MakeShared<FJsonObject>();
+		for (TFieldIterator<FProperty> PropIt(StructProp->Struct); PropIt; ++PropIt)
+		{
+			FProperty* InnerProp = *PropIt;
+			const void* InnerValuePtr = InnerProp->ContainerPtrToValuePtr<void>(ValuePtr);
+			StructObj->SetField(InnerProp->GetName(), DataTablePropertyToJsonValue(InnerProp, InnerValuePtr));
+		}
+		return MakeShared<FJsonValueObject>(StructObj);
+	}
+	else if (FArrayProperty* ArrayProp = CastField<FArrayProperty>(Property))
+	{
+		TArray<TSharedPtr<FJsonValue>> JsonArray;
+		FScriptArrayHelper ArrayHelper(ArrayProp, ValuePtr);
+		for (int32 i = 0; i < ArrayHelper.Num(); i++)
+		{
+			JsonArray.Add(DataTablePropertyToJsonValue(ArrayProp->Inner, ArrayHelper.GetRawPtr(i)));
+		}
+		return MakeShared<FJsonValueArray>(JsonArray);
+	}
+	else if (FMapProperty* MapProp = CastField<FMapProperty>(Property))
+	{
+		// Maps are exported as arrays of {key, value} objects
+		TArray<TSharedPtr<FJsonValue>> JsonArray;
+		FScriptMapHelper MapHelper(MapProp, ValuePtr);
+		for (int32 i = 0; i < MapHelper.Num(); i++)
+		{
+			if (MapHelper.IsValidIndex(i))
+			{
+				TSharedPtr<FJsonObject> EntryObj = MakeShared<FJsonObject>();
+				EntryObj->SetField(TEXT("key"), DataTablePropertyToJsonValue(MapProp->KeyProp, MapHelper.GetKeyPtr(i)));
+				EntryObj->SetField(TEXT("value"), DataTablePropertyToJsonValue(MapProp->ValueProp, MapHelper.GetValuePtr(i)));
+				JsonArray.Add(MakeShared<FJsonValueObject>(EntryObj));
+			}
+		}
+		return MakeShared<FJsonValueArray>(JsonArray);
+	}
+	else if (FSetProperty* SetProp = CastField<FSetProperty>(Property))
+	{
+		TArray<TSharedPtr<FJsonValue>> JsonArray;
+		FScriptSetHelper SetHelper(SetProp, ValuePtr);
+		for (int32 i = 0; i < SetHelper.Num(); i++)
+		{
+			if (SetHelper.IsValidIndex(i))
+			{
+				JsonArray.Add(DataTablePropertyToJsonValue(SetProp->ElementProp, SetHelper.GetElementPtr(i)));
+			}
+		}
+		return MakeShared<FJsonValueArray>(JsonArray);
+	}
+	else if (FSoftObjectProperty* SoftObjProp = CastField<FSoftObjectProperty>(Property))
+	{
+		const FSoftObjectPtr& SoftPtr = *static_cast<const FSoftObjectPtr*>(ValuePtr);
+		return MakeShared<FJsonValueString>(SoftPtr.ToString());
+	}
+	else if (FObjectProperty* ObjProp = CastField<FObjectProperty>(Property))
+	{
+		UObject* Obj = ObjProp->GetObjectPropertyValue(ValuePtr);
+		if (Obj)
+		{
+			return MakeShared<FJsonValueString>(Obj->GetPathName());
+		}
+		return MakeShared<FJsonValueString>(TEXT("None"));
+	}
+
+	// Fallback: export to string using UE's text export
+	FString StringValue;
+	Property->ExportTextItem_Direct(StringValue, ValuePtr, nullptr, nullptr, PPF_None);
+	return MakeShared<FJsonValueString>(StringValue);
+}
+
+//------------------------------------------------------------------------------
+// Helper: Convert a row to a JSON object
+//------------------------------------------------------------------------------
+static TSharedPtr<FJsonObject> RowToJsonObject(const UScriptStruct* RowStruct, const void* RowData)
+{
+	TSharedPtr<FJsonObject> RowObj = MakeShared<FJsonObject>();
+
+	for (TFieldIterator<FProperty> PropIt(RowStruct); PropIt; ++PropIt)
+	{
+		FProperty* Property = *PropIt;
+		const void* ValuePtr = Property->ContainerPtrToValuePtr<void>(RowData);
+		RowObj->SetField(Property->GetName(), DataTablePropertyToJsonValue(Property, ValuePtr));
+	}
+
+	return RowObj;
+}
+
+//------------------------------------------------------------------------------
+// Helper: Get property type as a friendly string
+//------------------------------------------------------------------------------
+static FString GetPropertyTypeString(FProperty* Property)
+{
+	if (FBoolProperty* BoolProp = CastField<FBoolProperty>(Property))
+	{
+		return TEXT("bool");
+	}
+	else if (FIntProperty* IntProp = CastField<FIntProperty>(Property))
+	{
+		return TEXT("int32");
+	}
+	else if (FInt64Property* Int64Prop = CastField<FInt64Property>(Property))
+	{
+		return TEXT("int64");
+	}
+	else if (FFloatProperty* FloatProp = CastField<FFloatProperty>(Property))
+	{
+		return TEXT("float");
+	}
+	else if (FDoubleProperty* DoubleProp = CastField<FDoubleProperty>(Property))
+	{
+		return TEXT("double");
+	}
+	else if (FStrProperty* StrProp = CastField<FStrProperty>(Property))
+	{
+		return TEXT("FString");
+	}
+	else if (FNameProperty* NameProp = CastField<FNameProperty>(Property))
+	{
+		return TEXT("FName");
+	}
+	else if (FTextProperty* TextProp = CastField<FTextProperty>(Property))
+	{
+		return TEXT("FText");
+	}
+	else if (FEnumProperty* EnumProp = CastField<FEnumProperty>(Property))
+	{
+		UEnum* Enum = EnumProp->GetEnum();
+		return FString::Printf(TEXT("enum:%s"), *Enum->GetName());
+	}
+	else if (FByteProperty* ByteProp = CastField<FByteProperty>(Property))
+	{
+		if (UEnum* Enum = ByteProp->GetIntPropertyEnum())
+		{
+			return FString::Printf(TEXT("enum:%s"), *Enum->GetName());
+		}
+		return TEXT("uint8");
+	}
+	else if (FStructProperty* StructProp = CastField<FStructProperty>(Property))
+	{
+		return FString::Printf(TEXT("struct:%s"), *StructProp->Struct->GetName());
+	}
+	else if (FArrayProperty* ArrayProp = CastField<FArrayProperty>(Property))
+	{
+		return FString::Printf(TEXT("array<%s>"), *GetPropertyTypeString(ArrayProp->Inner));
+	}
+	else if (FMapProperty* MapProp = CastField<FMapProperty>(Property))
+	{
+		return FString::Printf(TEXT("map<%s, %s>"), *GetPropertyTypeString(MapProp->KeyProp), *GetPropertyTypeString(MapProp->ValueProp));
+	}
+	else if (FSetProperty* SetProp = CastField<FSetProperty>(Property))
+	{
+		return FString::Printf(TEXT("set<%s>"), *GetPropertyTypeString(SetProp->ElementProp));
+	}
+	else if (FSoftObjectProperty* SoftObjProp = CastField<FSoftObjectProperty>(Property))
+	{
+		return FString::Printf(TEXT("soft_object:%s"), *SoftObjProp->PropertyClass->GetName());
+	}
+	else if (FObjectProperty* ObjProp = CastField<FObjectProperty>(Property))
+	{
+		return FString::Printf(TEXT("object:%s"), *ObjProp->PropertyClass->GetName());
+	}
+
+	return Property->GetCPPType();
+}
+
+//------------------------------------------------------------------------------
+// GetDataTableSchema
+//------------------------------------------------------------------------------
+
+FECACommandResult FECACommand_GetDataTableSchema::Execute(const TSharedPtr<FJsonObject>& Params)
+{
+	FString AssetPath;
+	if (!GetStringParam(Params, TEXT("asset_path"), AssetPath))
+	{
+		return FECACommandResult::Error(TEXT("Missing required parameter: asset_path"));
+	}
+
+	UDataTable* DataTable = LoadObject<UDataTable>(nullptr, *AssetPath);
+	if (!DataTable)
+	{
+		return FECACommandResult::Error(FString::Printf(TEXT("DataTable not found: %s"), *AssetPath));
+	}
+
+	const UScriptStruct* RowStruct = DataTable->GetRowStruct();
+	if (!RowStruct)
+	{
+		return FECACommandResult::Error(TEXT("DataTable has no row structure defined"));
+	}
+
+	TSharedPtr<FJsonObject> Result = MakeResult();
+	Result->SetStringField(TEXT("asset_path"), AssetPath);
+	Result->SetStringField(TEXT("row_struct_name"), RowStruct->GetName());
+	Result->SetStringField(TEXT("row_struct_path"), RowStruct->GetPathName());
+	Result->SetNumberField(TEXT("row_count"), DataTable->GetRowMap().Num());
+
+	// Build column definitions
+	TArray<TSharedPtr<FJsonValue>> ColumnsArray;
+	for (TFieldIterator<FProperty> PropIt(RowStruct); PropIt; ++PropIt)
+	{
+		FProperty* Property = *PropIt;
+		TSharedPtr<FJsonObject> ColumnObj = MakeShared<FJsonObject>();
+		ColumnObj->SetStringField(TEXT("name"), Property->GetName());
+		ColumnObj->SetStringField(TEXT("type"), GetPropertyTypeString(Property));
+		ColumnObj->SetStringField(TEXT("cpp_type"), Property->GetCPPType());
+
+		// Get tooltip/description if available
+		FString ToolTip = Property->GetMetaData(TEXT("ToolTip"));
+		if (!ToolTip.IsEmpty())
+		{
+			ColumnObj->SetStringField(TEXT("tooltip"), ToolTip);
+		}
+
+		// For enum properties, list the possible values
+		if (FEnumProperty* EnumProp = CastField<FEnumProperty>(Property))
+		{
+			UEnum* Enum = EnumProp->GetEnum();
+			TArray<TSharedPtr<FJsonValue>> EnumValues;
+			for (int32 i = 0; i < Enum->NumEnums() - 1; i++) // -1 to skip _MAX
+			{
+				EnumValues.Add(MakeShared<FJsonValueString>(Enum->GetNameStringByIndex(i)));
+			}
+			ColumnObj->SetArrayField(TEXT("enum_values"), EnumValues);
+		}
+		else if (FByteProperty* ByteProp = CastField<FByteProperty>(Property))
+		{
+			if (UEnum* Enum = ByteProp->GetIntPropertyEnum())
+			{
+				TArray<TSharedPtr<FJsonValue>> EnumValues;
+				for (int32 i = 0; i < Enum->NumEnums() - 1; i++)
+				{
+					EnumValues.Add(MakeShared<FJsonValueString>(Enum->GetNameStringByIndex(i)));
+				}
+				ColumnObj->SetArrayField(TEXT("enum_values"), EnumValues);
+			}
+		}
+
+		// For struct properties, list sub-fields
+		if (FStructProperty* StructProp = CastField<FStructProperty>(Property))
+		{
+			TArray<TSharedPtr<FJsonValue>> StructFields;
+			for (TFieldIterator<FProperty> InnerIt(StructProp->Struct); InnerIt; ++InnerIt)
+			{
+				TSharedPtr<FJsonObject> FieldInfo = MakeShared<FJsonObject>();
+				FieldInfo->SetStringField(TEXT("name"), InnerIt->GetName());
+				FieldInfo->SetStringField(TEXT("type"), GetPropertyTypeString(*InnerIt));
+				StructFields.Add(MakeShared<FJsonValueObject>(FieldInfo));
+			}
+			ColumnObj->SetArrayField(TEXT("struct_fields"), StructFields);
+		}
+
+		ColumnsArray.Add(MakeShared<FJsonValueObject>(ColumnObj));
+	}
+	Result->SetArrayField(TEXT("columns"), ColumnsArray);
+
+	return FECACommandResult::Success(Result);
+}
+
+//------------------------------------------------------------------------------
+// GetDataTableRows
+//------------------------------------------------------------------------------
+
+FECACommandResult FECACommand_GetDataTableRows::Execute(const TSharedPtr<FJsonObject>& Params)
+{
+	FString AssetPath;
+	if (!GetStringParam(Params, TEXT("asset_path"), AssetPath))
+	{
+		return FECACommandResult::Error(TEXT("Missing required parameter: asset_path"));
+	}
+
+	UDataTable* DataTable = LoadObject<UDataTable>(nullptr, *AssetPath);
+	if (!DataTable)
+	{
+		return FECACommandResult::Error(FString::Printf(TEXT("DataTable not found: %s"), *AssetPath));
+	}
+
+	const UScriptStruct* RowStruct = DataTable->GetRowStruct();
+	if (!RowStruct)
+	{
+		return FECACommandResult::Error(TEXT("DataTable has no row structure defined"));
+	}
+
+	// Check for specific row names filter
+	const TArray<TSharedPtr<FJsonValue>>* RowNamesArray = nullptr;
+	TSet<FString> RequestedRowNames;
+	bool bFilterByName = GetArrayParam(Params, TEXT("row_names"), RowNamesArray, false) && RowNamesArray && RowNamesArray->Num() > 0;
+	if (bFilterByName)
+	{
+		for (const TSharedPtr<FJsonValue>& Value : *RowNamesArray)
+		{
+			RequestedRowNames.Add(Value->AsString());
+		}
+	}
+
+	// Pagination
+	int32 Limit = -1;
+	int32 Offset = 0;
+	GetIntParam(Params, TEXT("limit"), Limit, false);
+	GetIntParam(Params, TEXT("offset"), Offset, false);
+
+	TSharedPtr<FJsonObject> Result = MakeResult();
+	Result->SetStringField(TEXT("asset_path"), AssetPath);
+	Result->SetStringField(TEXT("row_struct_name"), RowStruct->GetName());
+
+	const TMap<FName, uint8*>& RowMap = DataTable->GetRowMap();
+	Result->SetNumberField(TEXT("total_row_count"), RowMap.Num());
+
+	TArray<TSharedPtr<FJsonValue>> RowsArray;
+	int32 CurrentIndex = 0;
+	int32 ReturnedCount = 0;
+
+	for (const auto& Pair : RowMap)
+	{
+		FString RowName = Pair.Key.ToString();
+
+		// Apply name filter
+		if (bFilterByName && !RequestedRowNames.Contains(RowName))
+		{
+			continue;
+		}
+
+		// Apply offset
+		if (CurrentIndex < Offset)
+		{
+			CurrentIndex++;
+			continue;
+		}
+
+		// Apply limit
+		if (Limit > 0 && ReturnedCount >= Limit)
+		{
+			break;
+		}
+
+		TSharedPtr<FJsonObject> RowObj = MakeShared<FJsonObject>();
+		RowObj->SetStringField(TEXT("row_name"), RowName);
+		RowObj->SetObjectField(TEXT("data"), RowToJsonObject(RowStruct, Pair.Value));
+
+		RowsArray.Add(MakeShared<FJsonValueObject>(RowObj));
+		CurrentIndex++;
+		ReturnedCount++;
+	}
+
+	Result->SetArrayField(TEXT("rows"), RowsArray);
+	Result->SetNumberField(TEXT("returned_count"), ReturnedCount);
+
+	return FECACommandResult::Success(Result);
+}
+
+//------------------------------------------------------------------------------
+// SetDataTableRow
+//------------------------------------------------------------------------------
+
+FECACommandResult FECACommand_SetDataTableRow::Execute(const TSharedPtr<FJsonObject>& Params)
+{
+	FString AssetPath;
+	if (!GetStringParam(Params, TEXT("asset_path"), AssetPath))
+	{
+		return FECACommandResult::Error(TEXT("Missing required parameter: asset_path"));
+	}
+
+	FString RowName;
+	if (!GetStringParam(Params, TEXT("row_name"), RowName))
+	{
+		return FECACommandResult::Error(TEXT("Missing required parameter: row_name"));
+	}
+
+	const TSharedPtr<FJsonObject>* RowDataObj = nullptr;
+	if (!GetObjectParam(Params, TEXT("row_data"), RowDataObj))
+	{
+		return FECACommandResult::Error(TEXT("Missing required parameter: row_data"));
+	}
+
+	UDataTable* DataTable = LoadObject<UDataTable>(nullptr, *AssetPath);
+	if (!DataTable)
+	{
+		return FECACommandResult::Error(FString::Printf(TEXT("DataTable not found: %s"), *AssetPath));
+	}
+
+	const UScriptStruct* RowStruct = DataTable->GetRowStruct();
+	if (!RowStruct)
+	{
+		return FECACommandResult::Error(TEXT("DataTable has no row structure defined"));
+	}
+
+	// Check if row already exists
+	FName RowFName(*RowName);
+	bool bIsNewRow = (DataTable->GetRowMap().Find(RowFName) == nullptr);
+
+	// For editing, we need to work with the row data
+	// If the row exists, we modify it; if not, we add it
+	if (bIsNewRow)
+	{
+		// Add a new row - allocate memory for the struct and initialize it
+		uint8* NewRowData = (uint8*)FMemory::Malloc(RowStruct->GetStructureSize());
+		RowStruct->InitializeStruct(NewRowData);
+
+		// Set values from JSON
+		TArray<FString> FailedProperties;
+		for (const auto& JsonField : (*RowDataObj)->Values)
+		{
+			FProperty* Property = RowStruct->FindPropertyByName(FName(*JsonField.Key));
+			if (!Property)
+			{
+				FailedProperties.Add(FString::Printf(TEXT("%s (not found)"), *JsonField.Key));
+				continue;
+			}
+
+			void* ValuePtr = Property->ContainerPtrToValuePtr<void>(NewRowData);
+
+			// Convert JSON value to string for ImportText
+			FString ValueStr;
+			if (JsonField.Value->Type == EJson::Boolean)
+			{
+				ValueStr = JsonField.Value->AsBool() ? TEXT("True") : TEXT("False");
+			}
+			else if (JsonField.Value->Type == EJson::Number)
+			{
+				ValueStr = FString::SanitizeFloat(JsonField.Value->AsNumber());
+			}
+			else if (JsonField.Value->Type == EJson::String)
+			{
+				ValueStr = JsonField.Value->AsString();
+			}
+			else if (JsonField.Value->Type == EJson::Object || JsonField.Value->Type == EJson::Array)
+			{
+				// For complex types, serialize to JSON string and let ImportText parse it
+				TSharedRef<TJsonWriter<>> Writer = TJsonWriterFactory<>::Create(&ValueStr);
+				FJsonSerializer::Serialize(JsonField.Value, TEXT(""), Writer);
+			}
+
+			if (!ValueStr.IsEmpty())
+			{
+				const TCHAR* ImportResult = Property->ImportText_Direct(*ValueStr, ValuePtr, nullptr, PPF_None);
+				if (!ImportResult)
+				{
+					FailedProperties.Add(FString::Printf(TEXT("%s (import failed)"), *JsonField.Key));
+				}
+			}
+		}
+
+		// Add the row to the DataTable
+		DataTable->AddRow(RowFName, *reinterpret_cast<const FTableRowBase*>(NewRowData));
+
+		// Clean up our temporary allocation
+		RowStruct->DestroyStruct(NewRowData);
+		FMemory::Free(NewRowData);
+
+		DataTable->HandleDataTableChanged(RowFName);
+		DataTable->MarkPackageDirty();
+		FDataTableEditorUtils::BroadcastPostChange(DataTable, FDataTableEditorUtils::EDataTableChangeInfo::RowList);
+
+		TSharedPtr<FJsonObject> Result = MakeResult();
+		Result->SetStringField(TEXT("asset_path"), AssetPath);
+		Result->SetStringField(TEXT("row_name"), RowName);
+		Result->SetBoolField(TEXT("created"), true);
+		Result->SetBoolField(TEXT("updated"), false);
+
+		if (FailedProperties.Num() > 0)
+		{
+			TArray<TSharedPtr<FJsonValue>> FailedArray;
+			for (const FString& Failed : FailedProperties)
+			{
+				FailedArray.Add(MakeShared<FJsonValueString>(Failed));
+			}
+			Result->SetArrayField(TEXT("failed_properties"), FailedArray);
+		}
+
+		// Return the row as it was actually stored
+		uint8* StoredRow = DataTable->GetRowMap().FindRef(RowFName);
+		if (StoredRow)
+		{
+			Result->SetObjectField(TEXT("row_data"), RowToJsonObject(RowStruct, StoredRow));
+		}
+
+		return FECACommandResult::Success(Result);
+	}
+	else
+	{
+		// Update existing row
+		uint8* ExistingRowData = DataTable->GetRowMap().FindRef(RowFName);
+		if (!ExistingRowData)
+		{
+			return FECACommandResult::Error(FString::Printf(TEXT("Failed to find row '%s' for update"), *RowName));
+		}
+
+		DataTable->Modify();
+
+		TArray<FString> FailedProperties;
+		TArray<FString> UpdatedProperties;
+
+		for (const auto& JsonField : (*RowDataObj)->Values)
+		{
+			FProperty* Property = RowStruct->FindPropertyByName(FName(*JsonField.Key));
+			if (!Property)
+			{
+				FailedProperties.Add(FString::Printf(TEXT("%s (not found)"), *JsonField.Key));
+				continue;
+			}
+
+			void* ValuePtr = Property->ContainerPtrToValuePtr<void>(ExistingRowData);
+
+			FString ValueStr;
+			if (JsonField.Value->Type == EJson::Boolean)
+			{
+				ValueStr = JsonField.Value->AsBool() ? TEXT("True") : TEXT("False");
+			}
+			else if (JsonField.Value->Type == EJson::Number)
+			{
+				ValueStr = FString::SanitizeFloat(JsonField.Value->AsNumber());
+			}
+			else if (JsonField.Value->Type == EJson::String)
+			{
+				ValueStr = JsonField.Value->AsString();
+			}
+			else if (JsonField.Value->Type == EJson::Object || JsonField.Value->Type == EJson::Array)
+			{
+				TSharedRef<TJsonWriter<>> Writer = TJsonWriterFactory<>::Create(&ValueStr);
+				FJsonSerializer::Serialize(JsonField.Value, TEXT(""), Writer);
+			}
+
+			if (!ValueStr.IsEmpty())
+			{
+				const TCHAR* ImportResult = Property->ImportText_Direct(*ValueStr, ValuePtr, nullptr, PPF_None);
+				if (ImportResult)
+				{
+					UpdatedProperties.Add(FString(*JsonField.Key));
+				}
+				else
+				{
+					FailedProperties.Add(FString::Printf(TEXT("%s (import failed)"), *JsonField.Key));
+				}
+			}
+		}
+
+		DataTable->HandleDataTableChanged(RowFName);
+		DataTable->MarkPackageDirty();
+		FDataTableEditorUtils::BroadcastPostChange(DataTable, FDataTableEditorUtils::EDataTableChangeInfo::RowData);
+
+		TSharedPtr<FJsonObject> Result = MakeResult();
+		Result->SetStringField(TEXT("asset_path"), AssetPath);
+		Result->SetStringField(TEXT("row_name"), RowName);
+		Result->SetBoolField(TEXT("created"), false);
+		Result->SetBoolField(TEXT("updated"), true);
+
+		TArray<TSharedPtr<FJsonValue>> UpdatedArray;
+		for (const FString& Updated : UpdatedProperties)
+		{
+			UpdatedArray.Add(MakeShared<FJsonValueString>(Updated));
+		}
+		Result->SetArrayField(TEXT("updated_properties"), UpdatedArray);
+
+		if (FailedProperties.Num() > 0)
+		{
+			TArray<TSharedPtr<FJsonValue>> FailedArray;
+			for (const FString& Failed : FailedProperties)
+			{
+				FailedArray.Add(MakeShared<FJsonValueString>(Failed));
+			}
+			Result->SetArrayField(TEXT("failed_properties"), FailedArray);
+		}
+
+		// Return the updated row
+		Result->SetObjectField(TEXT("row_data"), RowToJsonObject(RowStruct, ExistingRowData));
+
+		return FECACommandResult::Success(Result);
+	}
+}
+
+//------------------------------------------------------------------------------
+// DeleteDataTableRow
+//------------------------------------------------------------------------------
+
+FECACommandResult FECACommand_DeleteDataTableRow::Execute(const TSharedPtr<FJsonObject>& Params)
+{
+	FString AssetPath;
+	if (!GetStringParam(Params, TEXT("asset_path"), AssetPath))
+	{
+		return FECACommandResult::Error(TEXT("Missing required parameter: asset_path"));
+	}
+
+	FString RowName;
+	if (!GetStringParam(Params, TEXT("row_name"), RowName))
+	{
+		return FECACommandResult::Error(TEXT("Missing required parameter: row_name"));
+	}
+
+	UDataTable* DataTable = LoadObject<UDataTable>(nullptr, *AssetPath);
+	if (!DataTable)
+	{
+		return FECACommandResult::Error(FString::Printf(TEXT("DataTable not found: %s"), *AssetPath));
+	}
+
+	FName RowFName(*RowName);
+
+	// Check if row exists
+	if (!DataTable->GetRowMap().Contains(RowFName))
+	{
+		return FECACommandResult::Error(FString::Printf(TEXT("Row '%s' not found in DataTable"), *RowName));
+	}
+
+	DataTable->Modify();
+	DataTable->RemoveRow(RowFName);
+	DataTable->HandleDataTableChanged(NAME_None);
+	DataTable->MarkPackageDirty();
+	FDataTableEditorUtils::BroadcastPostChange(DataTable, FDataTableEditorUtils::EDataTableChangeInfo::RowList);
+
+	TSharedPtr<FJsonObject> Result = MakeResult();
+	Result->SetStringField(TEXT("asset_path"), AssetPath);
+	Result->SetStringField(TEXT("deleted_row"), RowName);
+	Result->SetBoolField(TEXT("success"), true);
+	Result->SetNumberField(TEXT("remaining_row_count"), DataTable->GetRowMap().Num());
+
+	return FECACommandResult::Success(Result);
+}
diff --git a/Engine/Plugins/Experimental/ECABridge/Source/ECABridge/Private/Commands/ECAEditorCommands.cpp b/Engine/Plugins/Experimental/ECABridge/Source/ECABridge/Private/Commands/ECAEditorCommands.cpp
new file mode 100644
--- /dev/null
+++ b/Engine/Plugins/Experimental/ECABridge/Source/ECABridge/Private/Commands/ECAEditorCommands.cpp
@@ -0,0 +1,765 @@
+// Copyright Epic Games, Inc. All Rights Reserved.
+
+#include "Commands/ECAEditorCommands.h"
+#include "Engine/StaticMeshActor.h"
+#include "Commands/ECACommand.h"
+#include "Editor.h"
+#include "Engine/World.h"
+#include "EngineUtils.h"
+#include "Modules/ModuleManager.h"
+#include "LevelEditor.h"
+#include "IAssetViewport.h"
+#include "SLevelViewport.h"
+#include "EditorViewportClient.h"
+#include "Subsystems/EditorActorSubsystem.h"
+#include "Misc/FileHelper.h"
+#include "FileHelpers.h"
+#include "Misc/Base64.h"
+#include "ImageUtils.h"
+#include "Engine/GameViewportClient.h"
+#include "Kismet/GameplayStatics.h"
+#include "Components/SceneCaptureComponent2D.h"
+#include "Engine/SceneCapture2D.h"
+#include "Engine/TextureRenderTarget2D.h"
+#include "TextureResource.h"
+#include "HighResScreenshot.h"
+#include "UnrealClient.h"
+#include "Framework/Application/SlateApplication.h"
+#include "Widgets/Docking/SDockTab.h"
+
+// Register all editor commands
+REGISTER_ECA_COMMAND(FECACommand_FocusViewport)
+REGISTER_ECA_COMMAND(FECACommand_SelectActors)
+REGISTER_ECA_COMMAND(FECACommand_GetSelectedActors)
+
+REGISTER_ECA_COMMAND(FECACommand_TakeDepthScreenshot)
+REGISTER_ECA_COMMAND(FECACommand_TakeGameplayScreenshot)
+REGISTER_ECA_COMMAND(FECACommand_RunConsoleCommand)
+REGISTER_ECA_COMMAND(FECACommand_GetLevelInfo)
+REGISTER_ECA_COMMAND(FECACommand_OpenLevel)
+REGISTER_ECA_COMMAND(FECACommand_SaveLevel)
+REGISTER_ECA_COMMAND(FECACommand_PlayInEditor)
+REGISTER_ECA_COMMAND(FECACommand_StopPlayInEditor)
+REGISTER_ECA_COMMAND(FECACommand_GetPlayState)
+
+//------------------------------------------------------------------------------
+// FocusViewport
+//------------------------------------------------------------------------------
+
+FECACommandResult FECACommand_FocusViewport::Execute(const TSharedPtr<FJsonObject>& Params)
+{
+	// Collect actors to focus on
+	TArray<AActor*> ActorsToFocus;
+	
+	// Single actor
+	FString ActorName;
+	if (GetStringParam(Params, TEXT("actor_name"), ActorName, false))
+	{
+		AActor* Actor = FindActorByName(ActorName);
+		if (Actor)
+		{
+			ActorsToFocus.Add(Actor);
+		}
+	}
+	
+	// Multiple actors
+	const TArray<TSharedPtr<FJsonValue>>* ActorNamesArray = NULL;
+	if (GetArrayParam(Params, TEXT("actor_names"), ActorNamesArray, false) && ActorNamesArray)
+	{
+		for (const TSharedPtr<FJsonValue>& Value : *ActorNamesArray)
+		{
+			FString Name = Value->AsString();
+			AActor* Actor = FindActorByName(Name);
+			if (Actor)
+			{
+				ActorsToFocus.Add(Actor);
+			}
+		}
+	}
+	
+	if (ActorsToFocus.Num() > 0)
+	{
+		// Select and focus on actors
+		GEditor->SelectNone(true, true, false);
+		for (AActor* Actor : ActorsToFocus)
+		{
+			GEditor->SelectActor(Actor, true, true, false);
+		}
+		GEditor->MoveViewportCamerasToActor(ActorsToFocus, false);
+	}
+	else
+	{
+		// Focus on location
+		FVector Location;
+		if (GetVectorParam(Params, TEXT("location"), Location, false))
+		{
+			// Get viewport client
+			FLevelEditorModule& LevelEditorModule = FModuleManager::GetModuleChecked<FLevelEditorModule>(TEXT("LevelEditor"));
+			TSharedPtr<IAssetViewport> ActiveViewport = LevelEditorModule.GetFirstActiveViewport();
+			if (ActiveViewport.IsValid())
+			{
+				FEditorViewportClient& ViewportClient = ActiveViewport->GetAssetViewportClient();
+				ViewportClient.SetViewLocation(Location);
+			}
+		}
+		else
+		{
+			return FECACommandResult::Error(TEXT("No actor or location specified to focus on"));
+		}
+	}
+	
+	return FECACommandResult::Success();
+}
+
+//------------------------------------------------------------------------------
+// SelectActors
+//------------------------------------------------------------------------------
+
+FECACommandResult FECACommand_SelectActors::Execute(const TSharedPtr<FJsonObject>& Params)
+{
+	bool bDeselectAll = false;
+	GetBoolParam(Params, TEXT("deselect_all"), bDeselectAll, false);
+	
+	if (bDeselectAll)
+	{
+		GEditor->SelectNone(true, true, false);
+		return FECACommandResult::Success();
+	}
+	
+	bool bAddToSelection = false;
+	GetBoolParam(Params, TEXT("add_to_selection"), bAddToSelection, false);
+	
+	if (!bAddToSelection)
+	{
+		GEditor->SelectNone(true, true, false);
+	}
+	
+	TArray<FString> ActorNames;
+	
+	// Single actor
+	FString ActorName;
+	if (GetStringParam(Params, TEXT("actor_name"), ActorName, false))
+	{
+		ActorNames.Add(ActorName);
+	}
+	
+	// Multiple actors
+	const TArray<TSharedPtr<FJsonValue>>* ActorNamesArray = NULL;
+	if (GetArrayParam(Params, TEXT("actor_names"), ActorNamesArray, false) && ActorNamesArray)
+	{
+		for (const TSharedPtr<FJsonValue>& Value : *ActorNamesArray)
+		{
+			ActorNames.Add(Value->AsString());
+		}
+	}
+	
+	int32 SelectedCount = 0;
+	for (const FString& Name : ActorNames)
+	{
+		AActor* Actor = FindActorByName(Name);
+		if (Actor)
+		{
+			GEditor->SelectActor(Actor, true, true, false);
+			SelectedCount++;
+		}
+	}
+	
+	TSharedPtr<FJsonObject> Result = MakeResult();
+	Result->SetNumberField(TEXT("selected_count"), SelectedCount);
+	return FECACommandResult::Success(Result);
+}
+
+//------------------------------------------------------------------------------
+// GetSelectedActors
+//------------------------------------------------------------------------------
+
+FECACommandResult FECACommand_GetSelectedActors::Execute(const TSharedPtr<FJsonObject>& Params)
+{
+	UEditorActorSubsystem* EditorActorSubsystem = GEditor->GetEditorSubsystem<UEditorActorSubsystem>();
+	if (!EditorActorSubsystem)
+	{
+		return FECACommandResult::Error(TEXT("EditorActorSubsystem not available"));
+	}
+	
+	TArray<AActor*> SelectedActors = EditorActorSubsystem->GetSelectedLevelActors();
+	
+	TArray<TSharedPtr<FJsonValue>> ActorsArray;
+	for (AActor* Actor : SelectedActors)
+	{
+		TSharedPtr<FJsonObject> ActorJson = MakeShared<FJsonObject>();
+		ActorJson->SetStringField(TEXT("name"), Actor->GetActorLabel());
+		ActorJson->SetStringField(TEXT("path"), Actor->GetPathName());
+		ActorJson->SetStringField(TEXT("class"), Actor->GetClass()->GetName());
+		ActorJson->SetObjectField(TEXT("location"), VectorToJson(Actor->GetActorLocation()));
+		ActorJson->SetObjectField(TEXT("rotation"), RotatorToJson(Actor->GetActorRotation()));
+		ActorJson->SetObjectField(TEXT("scale"), VectorToJson(Actor->GetActorScale3D()));
+		
+		// If it's a static mesh actor, include the mesh path
+		if (AStaticMeshActor* SMActor = Cast<AStaticMeshActor>(Actor))
+		{
+			if (UStaticMeshComponent* SMComp = SMActor->GetStaticMeshComponent())
+			{
+				if (UStaticMesh* Mesh = SMComp->GetStaticMesh())
+				{
+					ActorJson->SetStringField(TEXT("mesh_path"), Mesh->GetPathName());
+				}
+			}
+		}
+		
+		ActorsArray.Add(MakeShared<FJsonValueObject>(ActorJson));
+	}
+	
+	TSharedPtr<FJsonObject> Result = MakeResult();
+	Result->SetArrayField(TEXT("actors"), ActorsArray);
+	Result->SetNumberField(TEXT("count"), ActorsArray.Num());
+	return FECACommandResult::Success(Result);
+}
+
+//------------------------------------------------------------------------------
+// TakeDepthScreenshot
+//------------------------------------------------------------------------------
+
+FECACommandResult FECACommand_TakeDepthScreenshot::Execute(const TSharedPtr<FJsonObject>& Params)
+{
+	// Get the level editor viewport
+	FLevelEditorModule& LevelEditorModule = FModuleManager::GetModuleChecked<FLevelEditorModule>(TEXT("LevelEditor"));
+	TSharedPtr<IAssetViewport> ActiveViewport = LevelEditorModule.GetFirstActiveViewport();
+	
+	if (!ActiveViewport.IsValid())
+	{
+		return FECACommandResult::Error(TEXT("No active viewport"));
+	}
+	
+	// Get viewport client for camera info
+	FEditorViewportClient* ViewportClient = static_cast<FEditorViewportClient*>(ActiveViewport->GetActiveViewport()->GetClient());
+	if (!ViewportClient)
+	{
+		return FECACommandResult::Error(TEXT("Could not get viewport client"));
+	}
+	
+	UWorld* World = GetEditorWorld();
+	if (!World)
+	{
+		return FECACommandResult::Error(TEXT("No world available"));
+	}
+	
+	// Get optional parameters
+	int32 Width = 1024;
+	int32 Height = 1024;
+	GetIntParam(Params, TEXT("width"), Width, false);
+	GetIntParam(Params, TEXT("height"), Height, false);
+	
+	double MaxDepth = 10000.0;  // 100 meters default
+	GetFloatParam(Params, TEXT("max_depth"), MaxDepth, false);
+	
+	// Get camera transform from viewport
+	FVector CameraLocation = ViewportClient->GetViewLocation();
+	FRotator CameraRotation = ViewportClient->GetViewRotation();
+	float FOV = ViewportClient->ViewFOV;
+	
+	// Create a temporary render target for depth capture
+	UTextureRenderTarget2D* RenderTarget = NewObject<UTextureRenderTarget2D>();
+	RenderTarget->InitCustomFormat(Width, Height, PF_FloatRGBA, false);
+	RenderTarget->UpdateResourceImmediate(true);
+	
+	// Create a temporary scene capture actor
+	FActorSpawnParameters SpawnParams;
+	SpawnParams.SpawnCollisionHandlingOverride = ESpawnActorCollisionHandlingMethod::AlwaysSpawn;
+	ASceneCapture2D* CaptureActor = World->SpawnActor<ASceneCapture2D>(CameraLocation, CameraRotation, SpawnParams);
+	
+	if (!CaptureActor)
+	{
+		return FECACommandResult::Error(TEXT("Failed to create scene capture actor"));
+	}
+	
+	// Configure capture component for depth
+	USceneCaptureComponent2D* CaptureComponent = CaptureActor->GetCaptureComponent2D();
+	CaptureComponent->TextureTarget = RenderTarget;
+	CaptureComponent->CaptureSource = ESceneCaptureSource::SCS_SceneDepth;
+	CaptureComponent->FOVAngle = FOV;
+	CaptureComponent->bCaptureEveryFrame = false;
+	CaptureComponent->bCaptureOnMovement = false;
+	CaptureComponent->MaxViewDistanceOverride = MaxDepth;
+	
+	// Capture the scene
+	CaptureComponent->CaptureScene();
+	
+	// Read pixels from render target
+	TArray<FFloat16Color> FloatPixels;
+	FTextureRenderTargetResource* RTResource = RenderTarget->GameThread_GetRenderTargetResource();
+	if (RTResource)
+	{
+		RTResource->ReadFloat16Pixels(FloatPixels);
+	}
+	
+	// Destroy the temporary actor
+	CaptureActor->Destroy();
+	
+	if (FloatPixels.Num() == 0)
+	{
+		return FECACommandResult::Error(TEXT("Failed to read depth pixels"));
+	}
+	
+	// Convert depth to grayscale image (normalized 0-255)
+	TArray<FColor> Bitmap;
+	Bitmap.SetNum(FloatPixels.Num());
+	
+	for (int32 i = 0; i < FloatPixels.Num(); i++)
+	{
+		// Get depth value (in R channel for SceneDepth)
+		float Depth = FloatPixels[i].R.GetFloat();
+		
+		// Normalize to 0-1 range based on max depth
+		float NormalizedDepth = FMath::Clamp(Depth / MaxDepth, 0.0f, 1.0f);
+		
+		// Convert to grayscale (0 = close, 255 = far)
+		uint8 Gray = static_cast<uint8>(NormalizedDepth * 255.0f);
+		
+		Bitmap[i] = FColor(Gray, Gray, Gray, 255);
+	}
+	
+	// Check if we should save to file
+	FString FilePath;
+	if (GetStringParam(Params, TEXT("file_path"), FilePath, false))
+	{
+		// Save to file
+		TArray64<uint8> CompressedData;
+		FImageUtils::PNGCompressImageArray(Width, Height, TArrayView64<const FColor>(Bitmap.GetData(), Bitmap.Num()), CompressedData);
+		
+		if (FFileHelper::SaveArrayToFile(TArrayView<const uint8>(CompressedData.GetData(), CompressedData.Num()), *FilePath))
+		{
+			TSharedPtr<FJsonObject> Result = MakeResult();
+			Result->SetStringField(TEXT("file_path"), FilePath);
+			Result->SetNumberField(TEXT("width"), Width);
+			Result->SetNumberField(TEXT("height"), Height);
+			Result->SetNumberField(TEXT("max_depth"), MaxDepth);
+			
+			// Include camera info
+			TSharedPtr<FJsonObject> CameraJson = MakeShared<FJsonObject>();
+			TSharedPtr<FJsonObject> LocationJson = MakeShared<FJsonObject>();
+			LocationJson->SetNumberField(TEXT("x"), CameraLocation.X);
+			LocationJson->SetNumberField(TEXT("y"), CameraLocation.Y);
+			LocationJson->SetNumberField(TEXT("z"), CameraLocation.Z);
+			CameraJson->SetObjectField(TEXT("location"), LocationJson);
+			
+			TSharedPtr<FJsonObject> RotationJson = MakeShared<FJsonObject>();
+			RotationJson->SetNumberField(TEXT("pitch"), CameraRotation.Pitch);
+			RotationJson->SetNumberField(TEXT("yaw"), CameraRotation.Yaw);
+			RotationJson->SetNumberField(TEXT("roll"), CameraRotation.Roll);
+			CameraJson->SetObjectField(TEXT("rotation"), RotationJson);
+			
+			CameraJson->SetNumberField(TEXT("fov"), FOV);
+			Result->SetObjectField(TEXT("camera"), CameraJson);
+			
+			return FECACommandResult::Success(Result);
+		}
+		else
+		{
+			return FECACommandResult::Error(FString::Printf(TEXT("Failed to save depth screenshot to: %s"), *FilePath));
+		}
+	}
+	else
+	{
+		// Return as base64
+		TArray64<uint8> CompressedData;
+		FImageUtils::PNGCompressImageArray(Width, Height, TArrayView64<const FColor>(Bitmap.GetData(), Bitmap.Num()), CompressedData);
+		FString Base64Image = FBase64::Encode(CompressedData.GetData(), CompressedData.Num());
+		
+		TSharedPtr<FJsonObject> Result = MakeResult();
+		Result->SetNumberField(TEXT("width"), Width);
+		Result->SetNumberField(TEXT("height"), Height);
+		Result->SetNumberField(TEXT("max_depth"), MaxDepth);
+		Result->SetStringField(TEXT("image_base64"), Base64Image);
+		Result->SetStringField(TEXT("format"), TEXT("png"));
+		
+		// Include camera info
+		TSharedPtr<FJsonObject> CameraJson = MakeShared<FJsonObject>();
+		TSharedPtr<FJsonObject> LocationJson = MakeShared<FJsonObject>();
+		LocationJson->SetNumberField(TEXT("x"), CameraLocation.X);
+		LocationJson->SetNumberField(TEXT("y"), CameraLocation.Y);
+		LocationJson->SetNumberField(TEXT("z"), CameraLocation.Z);
+		CameraJson->SetObjectField(TEXT("location"), LocationJson);
+		
+		TSharedPtr<FJsonObject> RotationJson = MakeShared<FJsonObject>();
+		RotationJson->SetNumberField(TEXT("pitch"), CameraRotation.Pitch);
+		RotationJson->SetNumberField(TEXT("yaw"), CameraRotation.Yaw);
+		RotationJson->SetNumberField(TEXT("roll"), CameraRotation.Roll);
+		CameraJson->SetObjectField(TEXT("rotation"), RotationJson);
+		
+		CameraJson->SetNumberField(TEXT("fov"), FOV);
+		Result->SetObjectField(TEXT("camera"), CameraJson);
+		
+		return FECACommandResult::Success(Result);
+	}
+}
+
+//------------------------------------------------------------------------------
+// TakeGameplayScreenshot
+//------------------------------------------------------------------------------
+
+FECACommandResult FECACommand_TakeGameplayScreenshot::Execute(const TSharedPtr<FJsonObject>& Params)
+{
+	// Determine output file path
+	FString FilePath;
+	if (!GetStringParam(Params, TEXT("file_path"), FilePath, false))
+	{
+		// Default: ProjectSaved/Screenshots/gameplay_YYYYMMDD_HHMMSS.png
+		FString Timestamp = FDateTime::Now().ToString(TEXT("%Y%m%d_%H%M%S"));
+		FilePath = FPaths::ConvertRelativePathToFull(
+			FPaths::ProjectSavedDir() / TEXT("Screenshots") / FString::Printf(TEXT("gameplay_%s.png"), *Timestamp));
+	}
+
+	// Determine capture target.
+	// Supported values: "pie", "editor_viewport", "editor_window"
+	// Default (empty): auto-detect -- PIE if playing, otherwise editor_viewport
+	FString Target;
+	GetStringParam(Params, TEXT("target"), Target, false);
+
+	if (Target.IsEmpty())
+	{
+		// Auto-detect
+		if (GEditor && GEditor->PlayWorld && GEngine && GEngine->GameViewport)
+		{
+			Target = TEXT("pie");
+		}
+		else
+		{
+			Target = TEXT("editor_viewport");
+		}
+	}
+
+	TArray<FColor> Bitmap;
+	int32 Width = 0;
+	int32 Height = 0;
+	FString Source;
+
+	if (Target == TEXT("pie"))
+	{
+		// Capture the PIE game window via Slate -- includes all Slate/UMG/HUD overlays.
+		if (!GEditor || !GEditor->PlayWorld || !GEngine || !GEngine->GameViewport)
+		{
+			return FECACommandResult::Error(TEXT("PIE is not running. Start Play in Editor first, or use 'editor_viewport' or 'editor_window' target."));
+		}
+
+		TSharedPtr<SWindow> GameWindow = GEngine->GameViewport->GetWindow();
+		if (!GameWindow.IsValid())
+		{
+			return FECACommandResult::Error(TEXT("PIE is running but no game window found."));
+		}
+
+		FIntVector Size;
+		bool bSuccess = FSlateApplication::Get().TakeScreenshot(GameWindow.ToSharedRef(), Bitmap, Size);
+
+		if (!bSuccess || Bitmap.Num() == 0)
+		{
+			return FECACommandResult::Error(TEXT("Failed to capture PIE screenshot via Slate. The window may not be fully rendered."));
+		}
+
+		Width = Size.X;
+		Height = Size.Y;
+		Source = TEXT("pie");
+	}
+	else if (Target == TEXT("editor_window"))
+	{
+		// Capture the full editor window via Slate -- includes all panels, tabs, graphs.
+		// Use the editor's own internal API: the LevelEditor module's tab knows exactly
+		// which SWindow the main editor frame lives in, regardless of focus or Z-order.
+		TSharedPtr<SWindow> EditorWindow;
+
+		FLevelEditorModule& LevelEditorModule = FModuleManager::GetModuleChecked<FLevelEditorModule>(TEXT("LevelEditor"));
+		TWeakPtr<SDockTab> LevelEditorTab = LevelEditorModule.GetLevelEditorInstanceTab();
+		if (LevelEditorTab.IsValid())
+		{
+			EditorWindow = LevelEditorTab.Pin()->GetParentWindow();
+		}
+
+		if (!EditorWindow.IsValid())
+		{
+			return FECACommandResult::Error(TEXT("No editor window found. Is the Level Editor tab loaded?"));
+		}
+
+		FIntVector Size;
+		bool bSuccess = FSlateApplication::Get().TakeScreenshot(EditorWindow.ToSharedRef(), Bitmap, Size);
+
+		if (!bSuccess || Bitmap.Num() == 0)
+		{
+			return FECACommandResult::Error(TEXT("Failed to capture editor window via Slate."));
+		}
+
+		Width = Size.X;
+		Height = Size.Y;
+		Source = TEXT("editor_window");
+	}
+	else if (Target == TEXT("editor_viewport"))
+	{
+		// Capture the editor 3D viewport via GetViewportScreenShot (render target only).
+		FLevelEditorModule& LevelEditorModule = FModuleManager::GetModuleChecked<FLevelEditorModule>(TEXT("LevelEditor"));
+		TSharedPtr<IAssetViewport> ActiveViewport = LevelEditorModule.GetFirstActiveViewport();
+
+		if (!ActiveViewport.IsValid())
+		{
+			return FECACommandResult::Error(TEXT("No active editor viewport available."));
+		}
+
+		FViewport* Viewport = ActiveViewport->GetActiveViewport();
+		if (!Viewport)
+		{
+			return FECACommandResult::Error(TEXT("Editor viewport is not available."));
+		}
+
+		FIntPoint ViewportSize = Viewport->GetSizeXY();
+		Width = ViewportSize.X;
+		Height = ViewportSize.Y;
+
+		if (Width == 0 || Height == 0)
+		{
+			return FECACommandResult::Error(TEXT("Editor viewport has zero size. Is it visible?"));
+		}
+
+		if (!GetViewportScreenShot(Viewport, Bitmap))
+		{
+			return FECACommandResult::Error(TEXT("Failed to capture editor viewport screenshot."));
+		}
+
+		Source = TEXT("editor_viewport");
+	}
+	else
+	{
+		return FECACommandResult::Error(FString::Printf(
+			TEXT("Unknown target '%s'. Valid targets: 'pie', 'editor_viewport', 'editor_window'."), *Target));
+	}
+
+	if (Bitmap.Num() == 0)
+	{
+		return FECACommandResult::Error(TEXT("Screenshot capture returned empty bitmap."));
+	}
+
+	// Force alpha to 255 (opaque). The editor viewport render target often has
+	// alpha=0 which makes PNGs appear black/transparent in most image viewers.
+	for (FColor& Pixel : Bitmap)
+	{
+		Pixel.A = 255;
+	}
+
+	// Ensure output directory exists
+	FString Directory = FPaths::GetPath(FilePath);
+	IPlatformFile::GetPlatformPhysical().CreateDirectoryTree(*Directory);
+
+	// Compress to PNG and save
+	TArray64<uint8> CompressedData;
+	FImageUtils::PNGCompressImageArray(Width, Height,
+		TArrayView64<const FColor>(Bitmap.GetData(), Bitmap.Num()), CompressedData);
+
+	if (CompressedData.Num() == 0)
+	{
+		return FECACommandResult::Error(TEXT("Failed to compress screenshot to PNG."));
+	}
+
+	if (!FFileHelper::SaveArrayToFile(
+		TArrayView<const uint8>(CompressedData.GetData(), CompressedData.Num()), *FilePath))
+	{
+		return FECACommandResult::Error(FString::Printf(TEXT("Failed to save screenshot to: %s"), *FilePath));
+	}
+
+	// Build result
+	TSharedPtr<FJsonObject> Result = MakeResult();
+	Result->SetStringField(TEXT("file_path"), FilePath);
+	Result->SetNumberField(TEXT("width"), Width);
+	Result->SetNumberField(TEXT("height"), Height);
+	Result->SetStringField(TEXT("source"), Source);
+
+	return FECACommandResult::Success(Result);
+}
+
+//------------------------------------------------------------------------------
+// RunConsoleCommand
+//------------------------------------------------------------------------------
+
+FECACommandResult FECACommand_RunConsoleCommand::Execute(const TSharedPtr<FJsonObject>& Params)
+{
+	FString Command;
+	if (!GetStringParam(Params, TEXT("command"), Command))
+	{
+		return FECACommandResult::Error(TEXT("Missing required parameter: command"));
+	}
+	
+	UWorld* World = GetEditorWorld();
+	if (World)
+	{
+		GEditor->Exec(World, *Command);
+	}
+	
+	TSharedPtr<FJsonObject> Result = MakeResult();
+	Result->SetStringField(TEXT("executed_command"), Command);
+	return FECACommandResult::Success(Result);
+}
+
+//------------------------------------------------------------------------------
+// GetLevelInfo
+//------------------------------------------------------------------------------
+
+FECACommandResult FECACommand_GetLevelInfo::Execute(const TSharedPtr<FJsonObject>& Params)
+{
+	UWorld* World = GetEditorWorld();
+	if (!World)
+	{
+		return FECACommandResult::Error(TEXT("No world available"));
+	}
+	
+	TSharedPtr<FJsonObject> Result = MakeResult();
+	Result->SetStringField(TEXT("level_name"), World->GetMapName());
+	Result->SetStringField(TEXT("level_path"), World->GetPathName());
+	
+	// Count actors by type
+	int32 TotalActors = 0;
+	TMap<FString, int32> ActorCounts;
+	
+	for (TActorIterator<AActor> It(World); It; ++It)
+	{
+		AActor* Actor = *It;
+		FString ClassName = Actor->GetClass()->GetName();
+		ActorCounts.FindOrAdd(ClassName)++;
+		TotalActors++;
+	}
+	
+	Result->SetNumberField(TEXT("total_actors"), TotalActors);
+	
+	TSharedPtr<FJsonObject> CountsObj = MakeShared<FJsonObject>();
+	for (const auto& Pair : ActorCounts)
+	{
+		CountsObj->SetNumberField(Pair.Key, Pair.Value);
+	}
+	Result->SetObjectField(TEXT("actor_counts"), CountsObj);
+	
+	return FECACommandResult::Success(Result);
+}
+
+//------------------------------------------------------------------------------
+// OpenLevel
+//------------------------------------------------------------------------------
+
+FECACommandResult FECACommand_OpenLevel::Execute(const TSharedPtr<FJsonObject>& Params)
+{
+	FString LevelPath;
+	if (!GetStringParam(Params, TEXT("level_path"), LevelPath))
+	{
+		return FECACommandResult::Error(TEXT("Missing required parameter: level_path"));
+	}
+	
+	// Ensure it's a valid level path
+	if (!LevelPath.StartsWith(TEXT("/Game/")))
+	{
+		LevelPath = TEXT("/Game/") + LevelPath;
+	}
+	
+	bool bSuccess = FEditorFileUtils::LoadMap(LevelPath);
+	
+	if (bSuccess)
+	{
+		TSharedPtr<FJsonObject> Result = MakeResult();
+		Result->SetStringField(TEXT("opened_level"), LevelPath);
+		return FECACommandResult::Success(Result);
+	}
+	else
+	{
+		return FECACommandResult::Error(FString::Printf(TEXT("Failed to open level: %s"), *LevelPath));
+	}
+}
+
+//------------------------------------------------------------------------------
+// SaveLevel
+//------------------------------------------------------------------------------
+
+FECACommandResult FECACommand_SaveLevel::Execute(const TSharedPtr<FJsonObject>& Params)
+{
+	UWorld* World = GetEditorWorld();
+	if (!World)
+	{
+		return FECACommandResult::Error(TEXT("No world available"));
+	}
+	
+	bool bSuccess = FEditorFileUtils::SaveCurrentLevel();
+	
+	if (bSuccess)
+	{
+		TSharedPtr<FJsonObject> Result = MakeResult();
+		Result->SetStringField(TEXT("saved_level"), World->GetMapName());
+		return FECACommandResult::Success(Result);
+	}
+	else
+	{
+		return FECACommandResult::Error(TEXT("Failed to save level"));
+	}
+}
+
+//------------------------------------------------------------------------------
+// PlayInEditor
+//------------------------------------------------------------------------------
+
+FECACommandResult FECACommand_PlayInEditor::Execute(const TSharedPtr<FJsonObject>& Params)
+{
+	if (GEditor->PlayWorld)
+	{
+		return FECACommandResult::Error(TEXT("Already playing in editor"));
+	}
+	
+	FLevelEditorModule& LevelEditorModule = FModuleManager::GetModuleChecked<FLevelEditorModule>(TEXT("LevelEditor"));
+	TSharedPtr<IAssetViewport> ActiveViewport = LevelEditorModule.GetFirstActiveViewport();
+	
+	FRequestPlaySessionParams SessionParams;
+	SessionParams.WorldType = EPlaySessionWorldType::PlayInEditor;
+	
+	GEditor->RequestPlaySession(SessionParams);
+	
+	return FECACommandResult::Success();
+}
+
+//------------------------------------------------------------------------------
+// StopPlayInEditor
+//------------------------------------------------------------------------------
+
+FECACommandResult FECACommand_StopPlayInEditor::Execute(const TSharedPtr<FJsonObject>& Params)
+{
+	if (!GEditor->PlayWorld)
+	{
+		return FECACommandResult::Error(TEXT("Not currently playing in editor"));
+	}
+	
+	GEditor->RequestEndPlayMap();
+	
+	return FECACommandResult::Success();
+}
+
+//------------------------------------------------------------------------------
+// GetPlayState
+//------------------------------------------------------------------------------
+
+FECACommandResult FECACommand_GetPlayState::Execute(const TSharedPtr<FJsonObject>& Params)
+{
+	TSharedPtr<FJsonObject> Result = MakeResult();
+	
+	bool bIsPlaying = GEditor->PlayWorld != nullptr;
+	bool bIsPaused = GEditor->PlayWorld ? GEditor->PlayWorld->IsPaused() : false;
+	bool bIsSimulating = GEditor->bIsSimulatingInEditor;
+	
+	Result->SetBoolField(TEXT("is_playing"), bIsPlaying);
+	Result->SetBoolField(TEXT("is_paused"), bIsPaused);
+	Result->SetBoolField(TEXT("is_simulating"), bIsSimulating);
+	
+	FString State;
+	if (bIsSimulating)
+	{
+		State = TEXT("Simulating");
+	}
+	else if (bIsPlaying)
+	{
+		State = bIsPaused ? TEXT("Paused") : TEXT("Playing");
+	}
+	else
+	{
+		State = TEXT("Stopped");
+	}
+	Result->SetStringField(TEXT("state"), State);
+	
+	return FECACommandResult::Success(Result);
+}
diff --git a/Engine/Plugins/Experimental/ECABridge/Source/ECABridge/Private/Commands/ECAEventCommands.cpp b/Engine/Plugins/Experimental/ECABridge/Source/ECABridge/Private/Commands/ECAEventCommands.cpp
new file mode 100644
--- /dev/null
+++ b/Engine/Plugins/Experimental/ECABridge/Source/ECABridge/Private/Commands/ECAEventCommands.cpp
@@ -0,0 +1,299 @@
+// Copyright Epic Games, Inc. All Rights Reserved.
+
+#include "Commands/ECACommand.h"
+#include "Dom/JsonObject.h"
+#include "Serialization/JsonSerializer.h"
+#include "HAL/PlatformTime.h"
+
+//------------------------------------------------------------------------------
+// Event Queue - Singleton for queuing events from editor callbacks
+//------------------------------------------------------------------------------
+
+class FECAEventQueue
+{
+public:
+	/** Maximum number of events to keep in the queue. Older events are dropped when limit is exceeded. */
+	static constexpr int32 MaxQueueSize = 10;
+	
+	/** Seconds to wait after startup before accepting events (avoids flood during editor load) */
+	static constexpr double StartupIgnoreSeconds = 5.0;
+	
+	static FECAEventQueue& Get()
+	{
+		static FECAEventQueue Instance;
+		return Instance;
+	}
+
+	/** Add an event to the queue */
+	void EnqueueEvent(const FString& EventType, const FString& Message, TSharedPtr<FJsonObject> Data = nullptr, bool bTriggerAIResponse = true)
+	{
+		// Ignore events during startup to avoid flood from editor loading
+		if (FPlatformTime::Seconds() - StartupTime < StartupIgnoreSeconds)
+		{
+			return;
+		}
+		
+		FScopeLock Lock(&QueueLock);
+		
+		TSharedPtr<FJsonObject> Event = MakeShared<FJsonObject>();
+		Event->SetStringField(TEXT("type"), EventType);
+		Event->SetStringField(TEXT("message"), Message);
+		Event->SetNumberField(TEXT("timestamp"), FPlatformTime::Seconds());
+		Event->SetBoolField(TEXT("trigger_ai_response"), bTriggerAIResponse);
+		
+		if (Data.IsValid())
+		{
+			Event->SetObjectField(TEXT("data"), Data);
+		}
+		
+		EventQueue.Add(Event);
+		
+		// If queue is too large, remove oldest events (from the front)
+		if (EventQueue.Num() > MaxQueueSize)
+		{
+			int32 NumToRemove = EventQueue.Num() - MaxQueueSize;
+			EventQueue.RemoveAt(0, NumToRemove);
+			UE_LOG(LogTemp, Log, TEXT("[ECABridge] Event queue overflow - dropped %d oldest events"), NumToRemove);
+		}
+		
+		UE_LOG(LogTemp, Log, TEXT("[ECABridge] Event queued: %s - %s"), *EventType, *Message);
+	}
+
+	/** Get and clear all pending events */
+	TArray<TSharedPtr<FJsonObject>> DequeueAll()
+	{
+		FScopeLock Lock(&QueueLock);
+		
+		TArray<TSharedPtr<FJsonObject>> Events = MoveTemp(EventQueue);
+		EventQueue.Empty();
+		
+		return Events;
+	}
+
+	/** Check if there are pending events */
+	bool HasPendingEvents() const
+	{
+		FScopeLock Lock(&QueueLock);
+		return EventQueue.Num() > 0;
+	}
+
+	/** Get the number of pending events */
+	int32 GetPendingCount() const
+	{
+		FScopeLock Lock(&QueueLock);
+		return EventQueue.Num();
+	}
+
+private:
+	FECAEventQueue()
+		: StartupTime(FPlatformTime::Seconds())
+	{
+	}
+	
+	TArray<TSharedPtr<FJsonObject>> EventQueue;
+	mutable FCriticalSection QueueLock;
+	double StartupTime;
+};
+
+//------------------------------------------------------------------------------
+// get_pending_events - Polled by ECA to retrieve queued events
+//------------------------------------------------------------------------------
+
+class FECACommand_GetPendingEvents : public IECACommand
+{
+public:
+	virtual FString GetName() const override { return TEXT("get_pending_events"); }
+	virtual FString GetDescription() const override 
+	{ 
+		return TEXT("Get and clear all pending events from ECABridge. Called by ECA to poll for editor events like selection changes, asset imports, etc."); 
+	}
+	virtual FString GetCategory() const override { return TEXT("Events"); }
+	
+	virtual TArray<FECACommandParam> GetParameters() const override
+	{
+		return TArray<FECACommandParam>();  // No parameters
+	}
+
+	virtual FECACommandResult Execute(const TSharedPtr<FJsonObject>& Params) override
+	{
+		TArray<TSharedPtr<FJsonObject>> Events = FECAEventQueue::Get().DequeueAll();
+		
+		TSharedPtr<FJsonObject> Result = MakeResult();
+		
+		// Convert events to JSON array
+		TArray<TSharedPtr<FJsonValue>> EventsArray;
+		for (const TSharedPtr<FJsonObject>& Event : Events)
+		{
+			EventsArray.Add(MakeShared<FJsonValueObject>(Event));
+		}
+		
+		Result->SetArrayField(TEXT("events"), EventsArray);
+		Result->SetStringField(TEXT("source"), TEXT("eca_bridge"));
+		Result->SetNumberField(TEXT("timestamp"), FPlatformTime::Seconds());
+		Result->SetNumberField(TEXT("count"), Events.Num());
+		
+		if (Events.Num() > 0)
+		{
+			UE_LOG(LogTemp, Log, TEXT("[ECABridge] Returned %d pending events"), Events.Num());
+		}
+		
+		return FECACommandResult::Success(Result);
+	}
+};
+
+//------------------------------------------------------------------------------
+// queue_event - Manually queue an event (for testing or programmatic use)
+//------------------------------------------------------------------------------
+
+class FECACommand_QueueEvent : public IECACommand
+{
+public:
+	virtual FString GetName() const override { return TEXT("queue_event"); }
+	virtual FString GetDescription() const override 
+	{ 
+		return TEXT("Manually queue an event to be sent to ECA. Useful for testing or programmatic event injection."); 
+	}
+	virtual FString GetCategory() const override { return TEXT("Events"); }
+	
+	virtual TArray<FECACommandParam> GetParameters() const override
+	{
+		return {
+			{ TEXT("event_type"), TEXT("string"), TEXT("Type of event (e.g., 'selection_changed', 'asset_imported')"), true },
+			{ TEXT("message"), TEXT("string"), TEXT("Human-readable description of the event"), true },
+			{ TEXT("data"), TEXT("object"), TEXT("Optional additional data for the event"), false },
+			{ TEXT("trigger_ai_response"), TEXT("boolean"), TEXT("Whether the AI should respond to this event (default: true)"), false }
+		};
+	}
+
+	virtual FECACommandResult Execute(const TSharedPtr<FJsonObject>& Params) override
+	{
+		FString EventType;
+		if (!GetStringParam(Params, TEXT("event_type"), EventType))
+		{
+			return FECACommandResult::Error(TEXT("Missing required parameter: event_type"));
+		}
+		
+		FString Message;
+		if (!GetStringParam(Params, TEXT("message"), Message))
+		{
+			return FECACommandResult::Error(TEXT("Missing required parameter: message"));
+		}
+		
+		bool bTriggerAIResponse = true;
+		GetBoolParam(Params, TEXT("trigger_ai_response"), bTriggerAIResponse, false);
+		
+		TSharedPtr<FJsonObject> Data;
+		const TSharedPtr<FJsonObject>* DataPtr;
+		if (GetObjectParam(Params, TEXT("data"), DataPtr, false))
+		{
+			Data = *DataPtr;
+		}
+		
+		FECAEventQueue::Get().EnqueueEvent(EventType, Message, Data, bTriggerAIResponse);
+		
+		TSharedPtr<FJsonObject> Result = MakeResult();
+		Result->SetStringField(TEXT("event_type"), EventType);
+		Result->SetStringField(TEXT("message"), Message);
+		Result->SetBoolField(TEXT("queued"), true);
+		Result->SetNumberField(TEXT("pending_count"), FECAEventQueue::Get().GetPendingCount());
+		
+		return FECACommandResult::Success(Result);
+	}
+};
+
+//------------------------------------------------------------------------------
+// get_event_queue_status - Check the event queue status
+//------------------------------------------------------------------------------
+
+class FECACommand_GetEventQueueStatus : public IECACommand
+{
+public:
+	virtual FString GetName() const override { return TEXT("get_event_queue_status"); }
+	virtual FString GetDescription() const override 
+	{ 
+		return TEXT("Get the current status of the event queue without clearing it."); 
+	}
+	virtual FString GetCategory() const override { return TEXT("Events"); }
+	
+	virtual TArray<FECACommandParam> GetParameters() const override
+	{
+		return TArray<FECACommandParam>();  // No parameters
+	}
+
+	virtual FECACommandResult Execute(const TSharedPtr<FJsonObject>& Params) override
+	{
+		TSharedPtr<FJsonObject> Result = MakeResult();
+		Result->SetBoolField(TEXT("has_pending_events"), FECAEventQueue::Get().HasPendingEvents());
+		Result->SetNumberField(TEXT("pending_count"), FECAEventQueue::Get().GetPendingCount());
+		Result->SetStringField(TEXT("source"), TEXT("eca_bridge"));
+		
+		return FECACommandResult::Success(Result);
+	}
+};
+
+//------------------------------------------------------------------------------
+// Register commands
+//------------------------------------------------------------------------------
+
+REGISTER_ECA_COMMAND(FECACommand_GetPendingEvents)
+REGISTER_ECA_COMMAND(FECACommand_QueueEvent)
+REGISTER_ECA_COMMAND(FECACommand_GetEventQueueStatus)
+
+
+//------------------------------------------------------------------------------
+// Helper functions for other modules to queue events
+//------------------------------------------------------------------------------
+
+namespace ECAEvents
+{
+	void QueueSelectionChanged(const FString& ActorName, const FString& ActorPath)
+	{
+		TSharedPtr<FJsonObject> Data = MakeShared<FJsonObject>();
+		Data->SetStringField(TEXT("actor_name"), ActorName);
+		Data->SetStringField(TEXT("actor_path"), ActorPath);
+		
+		FString Message = FString::Printf(TEXT("User selected '%s' in the Unreal Editor"), *ActorName);
+		FECAEventQueue::Get().EnqueueEvent(TEXT("selection_changed"), Message, Data, true);
+	}
+
+	void QueueAssetImported(const FString& AssetPath, const FString& AssetType)
+	{
+		TSharedPtr<FJsonObject> Data = MakeShared<FJsonObject>();
+		Data->SetStringField(TEXT("asset_path"), AssetPath);
+		Data->SetStringField(TEXT("asset_type"), AssetType);
+		
+		FString Message = FString::Printf(TEXT("Asset imported: %s (%s)"), *AssetPath, *AssetType);
+		FECAEventQueue::Get().EnqueueEvent(TEXT("asset_imported"), Message, Data, false);
+	}
+
+	void QueueAssetSaved(const FString& AssetPath)
+	{
+		TSharedPtr<FJsonObject> Data = MakeShared<FJsonObject>();
+		Data->SetStringField(TEXT("asset_path"), AssetPath);
+		
+		FString Message = FString::Printf(TEXT("Asset saved: %s"), *AssetPath);
+		FECAEventQueue::Get().EnqueueEvent(TEXT("asset_saved"), Message, Data, false);
+	}
+
+	void QueueBlueprintCompiled(const FString& BlueprintPath, bool bSuccess, const FString& ErrorMessage)
+	{
+		TSharedPtr<FJsonObject> Data = MakeShared<FJsonObject>();
+		Data->SetStringField(TEXT("blueprint_path"), BlueprintPath);
+		Data->SetBoolField(TEXT("success"), bSuccess);
+		if (!bSuccess)
+		{
+			Data->SetStringField(TEXT("error"), ErrorMessage);
+		}
+		
+		FString Message = bSuccess 
+			? FString::Printf(TEXT("Blueprint compiled successfully: %s"), *BlueprintPath)
+			: FString::Printf(TEXT("Blueprint compilation failed: %s - %s"), *BlueprintPath, *ErrorMessage);
+		
+		FECAEventQueue::Get().EnqueueEvent(TEXT("blueprint_compiled"), Message, Data, !bSuccess);  // Only trigger AI on failures
+	}
+
+	void QueueCustomEvent(const FString& EventType, const FString& Message, TSharedPtr<FJsonObject> Data, bool bTriggerAIResponse)
+	{
+		FECAEventQueue::Get().EnqueueEvent(EventType, Message, Data, bTriggerAIResponse);
+	}
+}
