ECABridge experimental MCP plugin

diff --git a/Engine/Plugins/Experimental/ECABridge/Source/ECABridge/Private/Commands/ECAMVVMCommands.cpp b/Engine/Plugins/Experimental/ECABridge/Source/ECABridge/Private/Commands/ECAMVVMCommands.cpp
new file mode 100644
--- /dev/null
+++ b/Engine/Plugins/Experimental/ECABridge/Source/ECABridge/Private/Commands/ECAMVVMCommands.cpp
@@ -0,0 +1,1116 @@
+// Copyright Epic Games, Inc. All Rights Reserved.
+
+#include "Commands/ECAMVVMCommands.h"
+#include "Commands/ECACommand.h"
+
+#include "Blueprint/WidgetTree.h"
+#include "Components/Image.h"
+#include "Components/RichTextBlock.h"
+#include "Components/TextBlock.h"
+#include "Components/Widget.h"
+#include "EdGraph/EdGraphPin.h"
+#include "Editor.h"
+#include "Kismet/BlueprintFunctionLibrary.h"
+#include "Kismet2/BlueprintEditorUtils.h"
+#include "Kismet2/KismetEditorUtilities.h"
+#include "MVVMBlueprintFunctionReference.h"
+#include "MVVMBlueprintPin.h"
+#include "MVVMBlueprintView.h"
+#include "MVVMBlueprintViewBinding.h"
+#include "MVVMBlueprintViewConversionFunction.h"
+#include "MVVMBlueprintViewModelContext.h"
+#include "MVVMEditorSubsystem.h"
+#include "MVVMPropertyPath.h"
+#include "MVVMWidgetBlueprintExtension_View.h"
+#include "Types/MVVMBindingMode.h"
+#include "Types/MVVMFieldVariant.h"
+#include "UObject/UObjectIterator.h"
+#include "WidgetBlueprint.h"
+
+// Register all MVVM commands
+REGISTER_ECA_COMMAND(FECACommand_SetWidgetViewmodel)
+REGISTER_ECA_COMMAND(FECACommand_BindTextToViewmodel)
+REGISTER_ECA_COMMAND(FECACommand_BindVisibilityToViewmodel)
+REGISTER_ECA_COMMAND(FECACommand_BindImageToViewmodel)
+REGISTER_ECA_COMMAND(FECACommand_GetViewmodelBindings)
+REGISTER_ECA_COMMAND(FECACommand_RemoveViewmodelBinding)
+REGISTER_ECA_COMMAND(FECACommand_GetWidgetMVVMProperties)
+REGISTER_ECA_COMMAND(FECACommand_AddViewmodelBinding)
+
+//------------------------------------------------------------------------------
+// Helper Functions
+//------------------------------------------------------------------------------
+
+static UWidgetBlueprint* LoadWidgetBlueprintByPath(const FString& WidgetPath)
+{
+	UObject* LoadedObject = LoadObject<UWidgetBlueprint>(nullptr, *WidgetPath);
+	if (!LoadedObject)
+	{
+		// Try appending object name for short paths
+		FString AlternatePath = WidgetPath + TEXT(".") + FPaths::GetBaseFilename(WidgetPath);
+		LoadedObject = LoadObject<UWidgetBlueprint>(nullptr, *AlternatePath);
+	}
+	return Cast<UWidgetBlueprint>(LoadedObject);
+}
+
+static UMVVMEditorSubsystem* GetMVVMSubsystem()
+{
+	return GEditor ? GEditor->GetEditorSubsystem<UMVVMEditorSubsystem>() : nullptr;
+}
+
+static UClass* FindViewModelClass(const FString& ViewModelClassName)
+{
+	// Search common script paths first
+	TArray<FString> SearchPaths;
+	SearchPaths.Add(FString::Printf(TEXT("/Script/ModelViewViewModel.%s"), *ViewModelClassName));
+	SearchPaths.Add(FString::Printf(TEXT("/Script/Engine.%s"), *ViewModelClassName));
+
+	for (const FString& Path : SearchPaths)
+	{
+		UClass* FoundClass = FindObject<UClass>(nullptr, *Path);
+		if (FoundClass)
+		{
+			return FoundClass;
+		}
+	}
+
+	// Iterate all classes as fallback
+	for (TObjectIterator<UClass> It; It; ++It)
+	{
+		if (It->GetName() == ViewModelClassName || It->GetAuthoredName() == ViewModelClassName)
+		{
+			return *It;
+		}
+	}
+
+	return nullptr;
+}
+
+static FProperty* FindPropertyOnClass(const UClass* Class, const FString& PropertyName)
+{
+	if (!Class)
+	{
+		return nullptr;
+	}
+
+	for (TFieldIterator<FProperty> It(Class); It; ++It)
+	{
+		if (It->GetName() == PropertyName || It->GetAuthoredName() == PropertyName)
+		{
+			return *It;
+		}
+	}
+
+	return nullptr;
+}
+
+static UWidget* FindWidgetInTree(UWidgetBlueprint* WidgetBP, const FString& WidgetName)
+{
+	if (!WidgetBP || !WidgetBP->WidgetTree)
+	{
+		return nullptr;
+	}
+
+	UWidget* FoundWidget = nullptr;
+	WidgetBP->WidgetTree->ForEachWidget([&](UWidget* Widget)
+	{
+		if (Widget && Widget->GetName() == WidgetName)
+		{
+			FoundWidget = Widget;
+		}
+	});
+	return FoundWidget;
+}
+
+static void CompileAndMarkDirty(UWidgetBlueprint* WidgetBP)
+{
+	FKismetEditorUtilities::CompileBlueprint(WidgetBP, EBlueprintCompileOptions::SkipGarbageCollection);
+	WidgetBP->MarkPackageDirty();
+}
+
+/**
+ * Common validation: load widget blueprint, get MVVM subsystem, get view with at least one ViewModel.
+ * Returns empty error string on success, error message on failure.
+ */
+static FString ValidateMVVMWidget(
+	const FString& WidgetPath,
+	UWidgetBlueprint*& OutWidgetBP,
+	UMVVMEditorSubsystem*& OutSubsystem,
+	UMVVMBlueprintView*& OutView)
+{
+	OutWidgetBP = LoadWidgetBlueprintByPath(WidgetPath);
+	if (!OutWidgetBP)
+	{
+		return FString::Printf(TEXT("Widget Blueprint not found: %s"), *WidgetPath);
+	}
+
+	OutSubsystem = GetMVVMSubsystem();
+	if (!OutSubsystem)
+	{
+		return TEXT("UMVVMEditorSubsystem not available. Is the editor running?");
+	}
+
+	OutView = OutSubsystem->GetView(OutWidgetBP);
+	if (!OutView || OutView->GetViewModels().Num() == 0)
+	{
+		return FString::Printf(TEXT("Widget '%s' has no MVVM ViewModel. Run set_widget_viewmodel first."), *WidgetPath);
+	}
+
+	return FString();
+}
+
+static FString BindingModeToString(EMVVMBindingMode Mode)
+{
+	switch (Mode)
+	{
+	case EMVVMBindingMode::OneWayToDestination:
+		return TEXT("OneWayToDestination");
+	case EMVVMBindingMode::OneWayToSource:
+		return TEXT("OneWayToSource");
+	case EMVVMBindingMode::TwoWay:
+		return TEXT("TwoWay");
+	case EMVVMBindingMode::OneTimeToDestination:
+		return TEXT("OneTimeToDestination");
+	case EMVVMBindingMode::OneTimeToSource:
+		return TEXT("OneTimeToSource");
+	default:
+		return TEXT("Unknown");
+	}
+}
+
+static EMVVMBindingMode ParseBindingMode(const FString& ModeStr)
+{
+	if (ModeStr.Equals(TEXT("TwoWay"), ESearchCase::IgnoreCase) || ModeStr.Equals(TEXT("twoway"), ESearchCase::IgnoreCase))
+	{
+		return EMVVMBindingMode::TwoWay;
+	}
+	if (ModeStr.Equals(TEXT("OneWayToSource"), ESearchCase::IgnoreCase))
+	{
+		return EMVVMBindingMode::OneWayToSource;
+	}
+	if (ModeStr.Equals(TEXT("OneTimeToDestination"), ESearchCase::IgnoreCase) || ModeStr.Equals(TEXT("onetime"), ESearchCase::IgnoreCase))
+	{
+		return EMVVMBindingMode::OneTimeToDestination;
+	}
+	// Default
+	return EMVVMBindingMode::OneWayToDestination;
+}
+
+static UFunction* ResolveConversionFunction(const UWidgetBlueprint* WidgetBP, const FString& FunctionSpec, FString& OutOwnerClass)
+{
+	OutOwnerClass = FString();
+
+	// Format 1: "ClassName::FunctionName" — explicit class qualification for disambiguation
+	FString ClassName, FuncName;
+	if (FunctionSpec.Split(TEXT("::"), &ClassName, &FuncName))
+	{
+		for (TObjectIterator<UClass> It; It; ++It)
+		{
+			if (It->GetName() == ClassName)
+			{
+				UFunction* Func = It->FindFunctionByName(*FuncName);
+				if (Func)
+				{
+					OutOwnerClass = It->GetName();
+					return Func;
+				}
+			}
+		}
+		return nullptr;
+	}
+
+	// Format 2: "FunctionName" — priority-based resolution:
+	// 1. Widget BP's own class (BP-local functions, highest priority)
+	if (WidgetBP && WidgetBP->GeneratedClass)
+	{
+		UFunction* Func = WidgetBP->GeneratedClass->FindFunctionByName(*FunctionSpec);
+		if (Func)
+		{
+			OutOwnerClass = WidgetBP->GeneratedClass->GetName();
+			return Func;
+		}
+	}
+
+	// 2. Non-BPFL classes (prefer concrete class functions over static library functions)
+	for (TObjectIterator<UClass> It; It; ++It)
+	{
+		if (It->IsChildOf(UBlueprintFunctionLibrary::StaticClass()))
+		{
+			continue;
+		}
+		UFunction* Func = It->FindFunctionByName(*FunctionSpec);
+		if (Func)
+		{
+			OutOwnerClass = It->GetName();
+			return Func;
+		}
+	}
+
+	// 3. BPFL classes (last resort — static library functions)
+	for (TObjectIterator<UClass> It; It; ++It)
+	{
+		if (!It->IsChildOf(UBlueprintFunctionLibrary::StaticClass()))
+		{
+			continue;
+		}
+		UFunction* Func = It->FindFunctionByName(*FunctionSpec);
+		if (Func)
+		{
+			OutOwnerClass = It->GetName();
+			return Func;
+		}
+	}
+
+	return nullptr;
+}
+
+//------------------------------------------------------------------------------
+// SetWidgetViewmodel
+//------------------------------------------------------------------------------
+
+FECACommandResult FECACommand_SetWidgetViewmodel::Execute(const TSharedPtr<FJsonObject>& Params)
+{
+	FString WidgetPath;
+	if (!GetStringParam(Params, TEXT("widget_path"), WidgetPath))
+	{
+		return FECACommandResult::Error(TEXT("Missing required parameter: widget_path"));
+	}
+
+	FString ViewModelClassName;
+	if (!GetStringParam(Params, TEXT("viewmodel_class"), ViewModelClassName))
+	{
+		return FECACommandResult::Error(TEXT("Missing required parameter: viewmodel_class"));
+	}
+
+	UWidgetBlueprint* WidgetBP = LoadWidgetBlueprintByPath(WidgetPath);
+	if (!WidgetBP)
+	{
+		return FECACommandResult::Error(FString::Printf(TEXT("Widget Blueprint not found: %s"), *WidgetPath));
+	}
+
+	UClass* VMClass = FindViewModelClass(ViewModelClassName);
+	if (!VMClass)
+	{
+		return FECACommandResult::Error(FString::Printf(
+			TEXT("ViewModel class not found: %s. Ensure the class exists and implements INotifyFieldValueChanged."),
+			*ViewModelClassName));
+	}
+
+	UMVVMEditorSubsystem* MVVMSubsystem = GetMVVMSubsystem();
+	if (!MVVMSubsystem)
+	{
+		return FECACommandResult::Error(TEXT("UMVVMEditorSubsystem not available. Is the editor running?"));
+	}
+
+	// Ensure the MVVM view extension exists on the widget
+	MVVMSubsystem->RequestView(WidgetBP);
+
+	FGuid ViewModelId = MVVMSubsystem->AddViewModel(WidgetBP, VMClass);
+	if (!ViewModelId.IsValid())
+	{
+		return FECACommandResult::Error(FString::Printf(
+			TEXT("Failed to add ViewModel %s. It may not implement INotifyFieldValueChanged."),
+			*ViewModelClassName));
+	}
+
+	CompileAndMarkDirty(WidgetBP);
+
+	TSharedPtr<FJsonObject> Result = MakeResult();
+	Result->SetStringField(TEXT("viewmodel_id"), ViewModelId.ToString());
+	Result->SetStringField(TEXT("viewmodel_class"), ViewModelClassName);
+	Result->SetStringField(TEXT("widget_path"), WidgetPath);
+	return FECACommandResult::Success(Result);
+}
+
+//------------------------------------------------------------------------------
+// BindTextToViewmodel
+//------------------------------------------------------------------------------
+
+FECACommandResult FECACommand_BindTextToViewmodel::Execute(const TSharedPtr<FJsonObject>& Params)
+{
+	FString WidgetPath, WidgetName, VMProperty;
+	if (!GetStringParam(Params, TEXT("widget_path"), WidgetPath) ||
+		!GetStringParam(Params, TEXT("widget_name"), WidgetName) ||
+		!GetStringParam(Params, TEXT("viewmodel_property"), VMProperty))
+	{
+		return FECACommandResult::Error(TEXT("Missing required parameters: widget_path, widget_name, viewmodel_property"));
+	}
+
+	UWidgetBlueprint* WidgetBP = nullptr;
+	UMVVMEditorSubsystem* MVVMSubsystem = nullptr;
+	UMVVMBlueprintView* View = nullptr;
+	FString ValidationError = ValidateMVVMWidget(WidgetPath, WidgetBP, MVVMSubsystem, View);
+	if (!ValidationError.IsEmpty())
+	{
+		return FECACommandResult::Error(ValidationError);
+	}
+
+	UWidget* TargetWidget = FindWidgetInTree(WidgetBP, WidgetName);
+	if (!TargetWidget)
+	{
+		return FECACommandResult::Error(FString::Printf(TEXT("Widget '%s' not found in the widget tree."), *WidgetName));
+	}
+
+	if (!Cast<UTextBlock>(TargetWidget) && !Cast<URichTextBlock>(TargetWidget))
+	{
+		return FECACommandResult::Error(FString::Printf(TEXT("Widget '%s' is not a TextBlock or RichTextBlock."), *WidgetName));
+	}
+
+	const FMVVMBlueprintViewModelContext& VMContext = View->GetViewModels()[0];
+	const UClass* VMClass = VMContext.GetViewModelClass();
+
+	FProperty* VMProp = FindPropertyOnClass(VMClass, VMProperty);
+	if (!VMProp)
+	{
+		return FECACommandResult::Error(FString::Printf(
+			TEXT("ViewModel property not found: %s on %s. Check property name spelling."),
+			*VMProperty, *VMClass->GetName()));
+	}
+
+	// Create binding
+	FMVVMBlueprintViewBinding& NewBinding = MVVMSubsystem->AddBinding(WidgetBP);
+
+	// Set binding mode
+	MVVMSubsystem->SetBindingTypeForBinding(WidgetBP, NewBinding, EMVVMBindingMode::OneWayToDestination);
+
+	// Set source path (ViewModel.Property)
+	FMVVMBlueprintPropertyPath SourcePath;
+	SourcePath.SetViewModelId(VMContext.GetViewModelId());
+	SourcePath.SetPropertyPath(WidgetBP, UE::MVVM::FMVVMConstFieldVariant(VMProp));
+	MVVMSubsystem->SetSourcePathForBinding(WidgetBP, NewBinding, SourcePath);
+
+	// Set destination path (Widget.Text)
+	FProperty* TextProp = FindPropertyOnClass(TargetWidget->GetClass(), TEXT("Text"));
+	if (!TextProp)
+	{
+		TextProp = FindPropertyOnClass(TargetWidget->GetClass(), TEXT("SetText"));
+	}
+
+	FMVVMBlueprintPropertyPath DestPath;
+	DestPath.SetWidgetName(TargetWidget->GetFName());
+	if (TextProp)
+	{
+		DestPath.SetPropertyPath(WidgetBP, UE::MVVM::FMVVMConstFieldVariant(TextProp));
+	}
+	MVVMSubsystem->SetDestinationPathForBinding(WidgetBP, NewBinding, DestPath, false);
+
+	CompileAndMarkDirty(WidgetBP);
+
+	TSharedPtr<FJsonObject> Result = MakeResult();
+	Result->SetStringField(TEXT("binding_mode"), TEXT("OneWayToDestination"));
+	return FECACommandResult::Success(Result);
+}
+
+//------------------------------------------------------------------------------
+// BindVisibilityToViewmodel
+//------------------------------------------------------------------------------
+
+FECACommandResult FECACommand_BindVisibilityToViewmodel::Execute(const TSharedPtr<FJsonObject>& Params)
+{
+	FString WidgetPath, WidgetName, VMProperty;
+	if (!GetStringParam(Params, TEXT("widget_path"), WidgetPath) ||
+		!GetStringParam(Params, TEXT("widget_name"), WidgetName) ||
+		!GetStringParam(Params, TEXT("viewmodel_property"), VMProperty))
+	{
+		return FECACommandResult::Error(TEXT("Missing required parameters: widget_path, widget_name, viewmodel_property"));
+	}
+
+	FString TrueVisibility = TEXT("Visible");
+	FString FalseVisibility = TEXT("Collapsed");
+	GetStringParam(Params, TEXT("true_visibility"), TrueVisibility, false);
+	GetStringParam(Params, TEXT("false_visibility"), FalseVisibility, false);
+
+	// Validate visibility enum values
+	TArray<FString> ValidVisibilities = {
+		TEXT("Visible"), TEXT("Collapsed"), TEXT("Hidden"),
+		TEXT("HitTestInvisible"), TEXT("SelfHitTestInvisible")
+	};
+
+	if (!ValidVisibilities.Contains(TrueVisibility))
+	{
+		return FECACommandResult::Error(FString::Printf(
+			TEXT("Invalid true_visibility: '%s'. Valid: Visible, Collapsed, Hidden, HitTestInvisible, SelfHitTestInvisible."),
+			*TrueVisibility));
+	}
+	if (!ValidVisibilities.Contains(FalseVisibility))
+	{
+		return FECACommandResult::Error(FString::Printf(
+			TEXT("Invalid false_visibility: '%s'. Valid: Visible, Collapsed, Hidden, HitTestInvisible, SelfHitTestInvisible."),
+			*FalseVisibility));
+	}
+	if (TrueVisibility == FalseVisibility)
+	{
+		return FECACommandResult::Error(TEXT("true_visibility and false_visibility must be different values."));
+	}
+
+	UWidgetBlueprint* WidgetBP = nullptr;
+	UMVVMEditorSubsystem* MVVMSubsystem = nullptr;
+	UMVVMBlueprintView* View = nullptr;
+	FString ValidationError = ValidateMVVMWidget(WidgetPath, WidgetBP, MVVMSubsystem, View);
+	if (!ValidationError.IsEmpty())
+	{
+		return FECACommandResult::Error(ValidationError);
+	}
+
+	UWidget* TargetWidget = FindWidgetInTree(WidgetBP, WidgetName);
+	if (!TargetWidget)
+	{
+		return FECACommandResult::Error(FString::Printf(TEXT("Widget '%s' not found in the widget tree."), *WidgetName));
+	}
+
+	const FMVVMBlueprintViewModelContext& VMContext = View->GetViewModels()[0];
+	const UClass* VMClass = VMContext.GetViewModelClass();
+
+	FProperty* VMProp = FindPropertyOnClass(VMClass, VMProperty);
+	if (!VMProp)
+	{
+		return FECACommandResult::Error(FString::Printf(
+			TEXT("ViewModel property not found: %s on %s. Check property name spelling."),
+			*VMProperty, *VMClass->GetName()));
+	}
+
+	// Create binding
+	FMVVMBlueprintViewBinding& NewBinding = MVVMSubsystem->AddBinding(WidgetBP);
+	MVVMSubsystem->SetBindingTypeForBinding(WidgetBP, NewBinding, EMVVMBindingMode::OneWayToDestination);
+
+	// Set destination path (Widget.Visibility)
+	FProperty* VisProp = FindPropertyOnClass(TargetWidget->GetClass(), TEXT("Visibility"));
+	if (!VisProp)
+	{
+		return FECACommandResult::Error(FString::Printf(TEXT("Widget '%s' does not have a Visibility property."), *WidgetName));
+	}
+
+	FMVVMBlueprintPropertyPath DestPath;
+	DestPath.SetWidgetName(TargetWidget->GetFName());
+	DestPath.SetPropertyPath(WidgetBP, UE::MVVM::FMVVMConstFieldVariant(VisProp));
+	MVVMSubsystem->SetDestinationPathForBinding(WidgetBP, NewBinding, DestPath, false);
+
+	// Find the engine's built-in Conv_BoolToSlateVisibility conversion function
+	UFunction* ConvFunc = nullptr;
+	for (TObjectIterator<UClass> It; It; ++It)
+	{
+		if (It->GetName() == TEXT("MVVMConversionLibrary"))
+		{
+			ConvFunc = It->FindFunctionByName(TEXT("Conv_BoolToSlateVisibility"));
+			break;
+		}
+	}
+
+	if (ConvFunc)
+	{
+		// Set the conversion function on the binding
+		FMVVMBlueprintFunctionReference FuncRef(WidgetBP, ConvFunc);
+		MVVMSubsystem->SetSourceToDestinationConversionFunction(WidgetBP, NewBinding, FuncRef);
+
+		// Wire ViewModel bool to the bIsVisible parameter pin
+		UE::MVVM::FConversionFunctionValue ConvValue(ConvFunc);
+		TArray<FMVVMBlueprintPinId> BoolPins = MVVMSubsystem->GetConversionFunctionPins(WidgetBP, ConvValue, VMProp);
+		if (BoolPins.Num() > 0)
+		{
+			FMVVMBlueprintPropertyPath VMSourcePath;
+			VMSourcePath.SetViewModelId(VMContext.GetViewModelId());
+			VMSourcePath.SetPropertyPath(WidgetBP, UE::MVVM::FMVVMConstFieldVariant(VMProp));
+			MVVMSubsystem->SetPathForConversionFunctionArgument(WidgetBP, NewBinding, BoolPins[0], VMSourcePath, true);
+		}
+
+		// Set TrueVisibility and FalseVisibility pin default values
+		UMVVMBlueprintViewConversionFunction* ConversionFunc = NewBinding.Conversion.GetConversionFunction(true);
+		if (ConversionFunc)
+		{
+			for (TFieldIterator<FProperty> ParamIt(ConvFunc); ParamIt && (ParamIt->PropertyFlags & CPF_Parm); ++ParamIt)
+			{
+				if (ParamIt->HasAnyPropertyFlags(CPF_ReturnParm))
+				{
+					continue;
+				}
+
+				FString ParamName = ParamIt->GetName();
+				FString DesiredDefault;
+
+				if (ParamName == TEXT("TrueVisibility"))
+				{
+					DesiredDefault = TrueVisibility;
+				}
+				else if (ParamName == TEXT("FalseVisibility"))
+				{
+					DesiredDefault = FalseVisibility;
+				}
+				else
+				{
+					continue;
+				}
+
+				TArray<FName> PinNames = { ParamIt->GetFName() };
+				FMVVMBlueprintPinId PinId(MoveTemp(PinNames));
+				UEdGraphPin* GraphPin = MVVMSubsystem->GetConversionFunctionArgumentPin(WidgetBP, NewBinding, PinId, true);
+				if (GraphPin)
+				{
+					GraphPin->DefaultValue = DesiredDefault;
+				}
+			}
+		}
+	}
+	else
+	{
+		// Fallback: direct binding without conversion function
+		FMVVMBlueprintPropertyPath SourcePath;
+		SourcePath.SetViewModelId(VMContext.GetViewModelId());
+		SourcePath.SetPropertyPath(WidgetBP, UE::MVVM::FMVVMConstFieldVariant(VMProp));
+		MVVMSubsystem->SetSourcePathForBinding(WidgetBP, NewBinding, SourcePath);
+	}
+
+	CompileAndMarkDirty(WidgetBP);
+
+	TSharedPtr<FJsonObject> Result = MakeResult();
+	Result->SetStringField(TEXT("conversion_function"), TEXT("Conv_BoolToSlateVisibility"));
+	Result->SetStringField(TEXT("true_visibility"), TrueVisibility);
+	Result->SetStringField(TEXT("false_visibility"), FalseVisibility);
+	return FECACommandResult::Success(Result);
+}
+
+//------------------------------------------------------------------------------
+// BindImageToViewmodel
+//------------------------------------------------------------------------------
+
+FECACommandResult FECACommand_BindImageToViewmodel::Execute(const TSharedPtr<FJsonObject>& Params)
+{
+	FString WidgetPath, WidgetName, VMProperty;
+	if (!GetStringParam(Params, TEXT("widget_path"), WidgetPath) ||
+		!GetStringParam(Params, TEXT("widget_name"), WidgetName) ||
+		!GetStringParam(Params, TEXT("viewmodel_property"), VMProperty))
+	{
+		return FECACommandResult::Error(TEXT("Missing required parameters: widget_path, widget_name, viewmodel_property"));
+	}
+
+	UWidgetBlueprint* WidgetBP = nullptr;
+	UMVVMEditorSubsystem* MVVMSubsystem = nullptr;
+	UMVVMBlueprintView* View = nullptr;
+	FString ValidationError = ValidateMVVMWidget(WidgetPath, WidgetBP, MVVMSubsystem, View);
+	if (!ValidationError.IsEmpty())
+	{
+		return FECACommandResult::Error(ValidationError);
+	}
+
+	UWidget* TargetWidget = FindWidgetInTree(WidgetBP, WidgetName);
+	if (!TargetWidget)
+	{
+		return FECACommandResult::Error(FString::Printf(TEXT("Widget '%s' not found in the widget tree."), *WidgetName));
+	}
+
+	if (!Cast<UImage>(TargetWidget))
+	{
+		return FECACommandResult::Error(FString::Printf(TEXT("Widget '%s' is not an Image widget."), *WidgetName));
+	}
+
+	const FMVVMBlueprintViewModelContext& VMContext = View->GetViewModels()[0];
+	const UClass* VMClass = VMContext.GetViewModelClass();
+
+	FProperty* VMProp = FindPropertyOnClass(VMClass, VMProperty);
+	if (!VMProp)
+	{
+		return FECACommandResult::Error(FString::Printf(
+			TEXT("ViewModel property not found: %s on %s. Check property name spelling."),
+			*VMProperty, *VMClass->GetName()));
+	}
+
+	// Create binding
+	FMVVMBlueprintViewBinding& NewBinding = MVVMSubsystem->AddBinding(WidgetBP);
+	MVVMSubsystem->SetBindingTypeForBinding(WidgetBP, NewBinding, EMVVMBindingMode::OneWayToDestination);
+
+	// Set source path (ViewModel.BrushProperty)
+	FMVVMBlueprintPropertyPath SourcePath;
+	SourcePath.SetViewModelId(VMContext.GetViewModelId());
+	SourcePath.SetPropertyPath(WidgetBP, UE::MVVM::FMVVMConstFieldVariant(VMProp));
+	MVVMSubsystem->SetSourcePathForBinding(WidgetBP, NewBinding, SourcePath);
+
+	// Set destination path (Image.Brush)
+	FProperty* BrushProp = FindPropertyOnClass(TargetWidget->GetClass(), TEXT("Brush"));
+	if (!BrushProp)
+	{
+		return FECACommandResult::Error(FString::Printf(TEXT("Widget '%s' does not have a Brush property."), *WidgetName));
+	}
+
+	FMVVMBlueprintPropertyPath DestPath;
+	DestPath.SetWidgetName(TargetWidget->GetFName());
+	DestPath.SetPropertyPath(WidgetBP, UE::MVVM::FMVVMConstFieldVariant(BrushProp));
+	MVVMSubsystem->SetDestinationPathForBinding(WidgetBP, NewBinding, DestPath, false);
+
+	CompileAndMarkDirty(WidgetBP);
+
+	TSharedPtr<FJsonObject> Result = MakeResult();
+	Result->SetStringField(TEXT("binding_mode"), TEXT("OneWayToDestination"));
+	return FECACommandResult::Success(Result);
+}
+
+//------------------------------------------------------------------------------
+// GetViewmodelBindings
+//------------------------------------------------------------------------------
+
+FECACommandResult FECACommand_GetViewmodelBindings::Execute(const TSharedPtr<FJsonObject>& Params)
+{
+	FString WidgetPath;
+	if (!GetStringParam(Params, TEXT("widget_path"), WidgetPath))
+	{
+		return FECACommandResult::Error(TEXT("Missing required parameter: widget_path"));
+	}
+
+	UWidgetBlueprint* WidgetBP = LoadWidgetBlueprintByPath(WidgetPath);
+	if (!WidgetBP)
+	{
+		return FECACommandResult::Error(FString::Printf(TEXT("Widget Blueprint not found: %s"), *WidgetPath));
+	}
+
+	UMVVMEditorSubsystem* MVVMSubsystem = GetMVVMSubsystem();
+	if (!MVVMSubsystem)
+	{
+		return FECACommandResult::Error(TEXT("UMVVMEditorSubsystem not available."));
+	}
+
+	UMVVMBlueprintView* View = MVVMSubsystem->GetView(WidgetBP);
+
+	// ViewModels
+	TArray<TSharedPtr<FJsonValue>> ViewModelArray;
+	if (View)
+	{
+		for (const FMVVMBlueprintViewModelContext& VMCtx : View->GetViewModels())
+		{
+			TSharedPtr<FJsonObject> VMObj = MakeShared<FJsonObject>();
+			VMObj->SetStringField(TEXT("id"), VMCtx.GetViewModelId().ToString());
+			VMObj->SetStringField(TEXT("name"), VMCtx.GetViewModelName().ToString());
+			VMObj->SetStringField(TEXT("class"), VMCtx.GetViewModelClass() ? VMCtx.GetViewModelClass()->GetName() : TEXT("null"));
+			ViewModelArray.Add(MakeShared<FJsonValueObject>(VMObj));
+		}
+	}
+
+	// Bindings
+	TArray<TSharedPtr<FJsonValue>> BindingArray;
+	if (View)
+	{
+		for (int32 i = 0; i < View->GetNumBindings(); ++i)
+		{
+			const FMVVMBlueprintViewBinding* Binding = View->GetBindingAt(i);
+			if (!Binding)
+			{
+				continue;
+			}
+
+			TSharedPtr<FJsonObject> BindObj = MakeShared<FJsonObject>();
+			BindObj->SetStringField(TEXT("display"), Binding->GetDisplayNameString(WidgetBP));
+			BindObj->SetStringField(TEXT("mode"), BindingModeToString(Binding->BindingType));
+			BindObj->SetBoolField(TEXT("enabled"), Binding->bEnabled);
+			BindObj->SetNumberField(TEXT("index"), i);
+			BindingArray.Add(MakeShared<FJsonValueObject>(BindObj));
+		}
+	}
+
+	TSharedPtr<FJsonObject> Result = MakeResult();
+	Result->SetArrayField(TEXT("viewmodels"), ViewModelArray);
+	Result->SetArrayField(TEXT("bindings"), BindingArray);
+	Result->SetNumberField(TEXT("viewmodel_count"), ViewModelArray.Num());
+	Result->SetNumberField(TEXT("binding_count"), BindingArray.Num());
+	return FECACommandResult::Success(Result);
+}
+
+//------------------------------------------------------------------------------
+// RemoveViewmodelBinding
+//------------------------------------------------------------------------------
+
+FECACommandResult FECACommand_RemoveViewmodelBinding::Execute(const TSharedPtr<FJsonObject>& Params)
+{
+	FString WidgetPath;
+	if (!GetStringParam(Params, TEXT("widget_path"), WidgetPath))
+	{
+		return FECACommandResult::Error(TEXT("Missing required parameter: widget_path"));
+	}
+
+	UWidgetBlueprint* WidgetBP = nullptr;
+	UMVVMEditorSubsystem* MVVMSubsystem = nullptr;
+	UMVVMBlueprintView* View = nullptr;
+	FString ValidationError = ValidateMVVMWidget(WidgetPath, WidgetBP, MVVMSubsystem, View);
+	if (!ValidationError.IsEmpty())
+	{
+		return FECACommandResult::Error(ValidationError);
+	}
+
+	// Remove by index
+	int32 BindingIndex = -1;
+	if (GetIntParam(Params, TEXT("binding_index"), BindingIndex, false))
+	{
+		if (BindingIndex < 0 || BindingIndex >= View->GetNumBindings())
+		{
+			return FECACommandResult::Error(FString::Printf(
+				TEXT("Binding index %d out of range. Widget has %d bindings."),
+				BindingIndex, View->GetNumBindings()));
+		}
+
+		const FMVVMBlueprintViewBinding* Binding = View->GetBindingAt(BindingIndex);
+		if (Binding)
+		{
+			MVVMSubsystem->RemoveBinding(WidgetBP, *Binding);
+		}
+
+		CompileAndMarkDirty(WidgetBP);
+
+		TSharedPtr<FJsonObject> Result = MakeResult();
+		Result->SetNumberField(TEXT("binding_index"), BindingIndex);
+		return FECACommandResult::Success(Result);
+	}
+
+	// Remove by source/destination match
+	FString SourceProperty, DestinationProperty;
+	GetStringParam(Params, TEXT("source_property"), SourceProperty, false);
+	GetStringParam(Params, TEXT("destination_property"), DestinationProperty, false);
+
+	if (!SourceProperty.IsEmpty() || !DestinationProperty.IsEmpty())
+	{
+		int32 RemovedCount = 0;
+		for (int32 i = View->GetNumBindings() - 1; i >= 0; --i)
+		{
+			const FMVVMBlueprintViewBinding* Binding = View->GetBindingAt(i);
+			if (!Binding)
+			{
+				continue;
+			}
+
+			FString DisplayStr = Binding->GetDisplayNameString(WidgetBP);
+
+			bool bMatch = true;
+			if (!SourceProperty.IsEmpty() && !DisplayStr.Contains(SourceProperty))
+			{
+				bMatch = false;
+			}
+			if (!DestinationProperty.IsEmpty() && !DisplayStr.Contains(DestinationProperty))
+			{
+				bMatch = false;
+			}
+
+			if (bMatch)
+			{
+				MVVMSubsystem->RemoveBinding(WidgetBP, *Binding);
+				++RemovedCount;
+			}
+		}
+
+		CompileAndMarkDirty(WidgetBP);
+
+		TSharedPtr<FJsonObject> Result = MakeResult();
+		Result->SetNumberField(TEXT("removed_count"), RemovedCount);
+		return FECACommandResult::Success(Result);
+	}
+
+	return FECACommandResult::Error(TEXT("Provide either binding_index or source_property/destination_property to identify the binding."));
+}
+
+//------------------------------------------------------------------------------
+// GetWidgetMVVMProperties
+//------------------------------------------------------------------------------
+
+FECACommandResult FECACommand_GetWidgetMVVMProperties::Execute(const TSharedPtr<FJsonObject>& Params)
+{
+	FString WidgetPath, WidgetName;
+	if (!GetStringParam(Params, TEXT("widget_path"), WidgetPath) ||
+		!GetStringParam(Params, TEXT("widget_name"), WidgetName))
+	{
+		return FECACommandResult::Error(TEXT("Missing required parameters: widget_path, widget_name"));
+	}
+
+	UWidgetBlueprint* WidgetBP = LoadWidgetBlueprintByPath(WidgetPath);
+	if (!WidgetBP)
+	{
+		return FECACommandResult::Error(FString::Printf(TEXT("Widget Blueprint not found: %s"), *WidgetPath));
+	}
+
+	UWidget* TargetWidget = FindWidgetInTree(WidgetBP, WidgetName);
+	if (!TargetWidget)
+	{
+		return FECACommandResult::Error(FString::Printf(TEXT("Widget '%s' not found in the widget tree."), *WidgetName));
+	}
+
+	TArray<TSharedPtr<FJsonValue>> PropertiesArray;
+	for (TFieldIterator<FProperty> It(TargetWidget->GetClass()); It; ++It)
+	{
+		FProperty* Prop = *It;
+		bool bBindable = Prop->HasAllPropertyFlags(CPF_BlueprintVisible);
+
+		TSharedPtr<FJsonObject> PropObj = MakeShared<FJsonObject>();
+		PropObj->SetStringField(TEXT("name"), Prop->GetName());
+		PropObj->SetStringField(TEXT("type"), Prop->GetCPPType());
+		PropObj->SetBoolField(TEXT("bindable"), bBindable);
+		PropertiesArray.Add(MakeShared<FJsonValueObject>(PropObj));
+	}
+
+	TSharedPtr<FJsonObject> Result = MakeResult();
+	Result->SetArrayField(TEXT("properties"), PropertiesArray);
+	Result->SetStringField(TEXT("widget_name"), WidgetName);
+	Result->SetStringField(TEXT("widget_class"), TargetWidget->GetClass()->GetName());
+	Result->SetNumberField(TEXT("property_count"), PropertiesArray.Num());
+	return FECACommandResult::Success(Result);
+}
+
+//------------------------------------------------------------------------------
+// AddViewmodelBinding
+//------------------------------------------------------------------------------
+
+FECACommandResult FECACommand_AddViewmodelBinding::Execute(const TSharedPtr<FJsonObject>& Params)
+{
+	FString WidgetPath, SourceProperty, DestinationProperty;
+	if (!GetStringParam(Params, TEXT("widget_path"), WidgetPath) ||
+		!GetStringParam(Params, TEXT("source_property"), SourceProperty) ||
+		!GetStringParam(Params, TEXT("destination_property"), DestinationProperty))
+	{
+		return FECACommandResult::Error(TEXT("Missing required parameters: widget_path, source_property, destination_property"));
+	}
+
+	FString DestinationWidget = TEXT("self");
+	GetStringParam(Params, TEXT("destination_widget"), DestinationWidget, false);
+
+	FString ViewModelName;
+	bool bHasViewModelName = GetStringParam(Params, TEXT("viewmodel_name"), ViewModelName, false);
+
+	FString BindingModeStr = TEXT("OneWayToDestination");
+	GetStringParam(Params, TEXT("binding_mode"), BindingModeStr, false);
+
+	FString ConversionFunctionStr;
+	bool bHasConversionFunction = GetStringParam(Params, TEXT("conversion_function"), ConversionFunctionStr, false);
+
+	// Load and validate
+	UWidgetBlueprint* WidgetBP = nullptr;
+	UMVVMEditorSubsystem* MVVMSubsystem = nullptr;
+	UMVVMBlueprintView* View = nullptr;
+	FString ValidationError = ValidateMVVMWidget(WidgetPath, WidgetBP, MVVMSubsystem, View);
+	if (!ValidationError.IsEmpty())
+	{
+		return FECACommandResult::Error(ValidationError);
+	}
+
+	// Find ViewModel by name or fall back to first
+	const FMVVMBlueprintViewModelContext* VMContext = nullptr;
+	if (bHasViewModelName && !ViewModelName.IsEmpty())
+	{
+		VMContext = View->FindViewModel(FName(*ViewModelName));
+		if (!VMContext)
+		{
+			// Build list of available VM names for the error message
+			FString AvailableVMs;
+			for (const FMVVMBlueprintViewModelContext& Ctx : View->GetViewModels())
+			{
+				if (!AvailableVMs.IsEmpty())
+				{
+					AvailableVMs += TEXT(", ");
+				}
+				AvailableVMs += Ctx.GetViewModelName().ToString();
+			}
+			return FECACommandResult::Error(FString::Printf(
+				TEXT("ViewModel '%s' not found. Available: %s"),
+				*ViewModelName, *AvailableVMs));
+		}
+	}
+	else
+	{
+		VMContext = &View->GetViewModels()[0];
+	}
+
+	const UClass* VMClass = VMContext->GetViewModelClass();
+	if (!VMClass)
+	{
+		return FECACommandResult::Error(TEXT("ViewModel class is null."));
+	}
+
+	// Resolve source property on ViewModel (FProperty expected, UFunction as fallback)
+	FProperty* VMProp = FindPropertyOnClass(VMClass, SourceProperty);
+	UFunction* VMFunc = nullptr;
+	if (!VMProp)
+	{
+		VMFunc = VMClass->FindFunctionByName(*SourceProperty);
+	}
+	if (!VMProp && !VMFunc)
+	{
+		return FECACommandResult::Error(FString::Printf(
+			TEXT("ViewModel property not found: '%s' on %s. Check property name spelling."),
+			*SourceProperty, *VMClass->GetName()));
+	}
+	UE::MVVM::FMVVMConstFieldVariant VMField = VMProp
+		? UE::MVVM::FMVVMConstFieldVariant(VMProp)
+		: UE::MVVM::FMVVMConstFieldVariant(VMFunc);
+
+	// Resolve destination: self or child widget
+	bool bIsSelfBinding = DestinationWidget.IsEmpty() || DestinationWidget.Equals(TEXT("self"), ESearchCase::IgnoreCase);
+	UWidget* TargetWidget = nullptr;
+	const UClass* DestClass = nullptr;
+
+	if (bIsSelfBinding)
+	{
+		DestClass = WidgetBP->GeneratedClass;
+		if (!DestClass)
+		{
+			return FECACommandResult::Error(TEXT("Widget Blueprint has no generated class. Try compiling it first."));
+		}
+	}
+	else
+	{
+		TargetWidget = FindWidgetInTree(WidgetBP, DestinationWidget);
+		if (!TargetWidget)
+		{
+			return FECACommandResult::Error(FString::Printf(
+				TEXT("Widget '%s' not found in the widget tree."), *DestinationWidget));
+		}
+		DestClass = TargetWidget->GetClass();
+	}
+
+	// Resolve destination property (FProperty first, then UFunction for setter-based bindings)
+	FProperty* DestProp = FindPropertyOnClass(DestClass, DestinationProperty);
+	UFunction* DestFunc = nullptr;
+	if (!DestProp)
+	{
+		DestFunc = DestClass->FindFunctionByName(*DestinationProperty);
+	}
+	if (!DestProp && !DestFunc)
+	{
+		return FECACommandResult::Error(FString::Printf(
+			TEXT("Destination property or function not found: '%s' on %s. Check spelling."),
+			*DestinationProperty, *DestClass->GetName()));
+	}
+	UE::MVVM::FMVVMConstFieldVariant DestField = DestProp
+		? UE::MVVM::FMVVMConstFieldVariant(DestProp)
+		: UE::MVVM::FMVVMConstFieldVariant(DestFunc);
+
+	// Parse binding mode
+	EMVVMBindingMode BindingMode = ParseBindingMode(BindingModeStr);
+
+	// Create binding
+	FMVVMBlueprintViewBinding& NewBinding = MVVMSubsystem->AddBinding(WidgetBP);
+
+	// Set binding mode
+	MVVMSubsystem->SetBindingTypeForBinding(WidgetBP, NewBinding, BindingMode);
+
+	// Build and set source path (ViewModel.Property)
+	FMVVMBlueprintPropertyPath SourcePath;
+	SourcePath.SetViewModelId(VMContext->GetViewModelId());
+	SourcePath.SetPropertyPath(WidgetBP, VMField);
+
+	// Build and set destination path (Widget.Property or Self.Property)
+	FMVVMBlueprintPropertyPath DestPath;
+	if (bIsSelfBinding)
+	{
+		DestPath.SetSelfContext();
+	}
+	else
+	{
+		DestPath.SetWidgetName(TargetWidget->GetFName());
+	}
+	DestPath.SetPropertyPath(WidgetBP, DestField);
+
+	// Handle conversion function if specified
+	FString ConversionFunctionUsed;
+	FString ConversionFunctionOwner;
+	if (bHasConversionFunction && !ConversionFunctionStr.IsEmpty())
+	{
+		UFunction* ConvFunc = ResolveConversionFunction(WidgetBP, ConversionFunctionStr, ConversionFunctionOwner);
+		if (!ConvFunc)
+		{
+			// Clean up the binding we already added before returning error
+			MVVMSubsystem->RemoveBinding(WidgetBP, NewBinding);
+			return FECACommandResult::Error(FString::Printf(
+				TEXT("Conversion function not found: '%s'"), *ConversionFunctionStr));
+		}
+
+		// Pre-validate: check conversion function signature is compatible with source/destination types
+		// VMProp = expected input type (null skips check), DestProp = expected output type (null for UFunction dests, skips check)
+		if (!MVVMSubsystem->IsValidConversionFunction(WidgetBP, ConvFunc, VMProp, DestProp))
+		{
+			MVVMSubsystem->RemoveBinding(WidgetBP, NewBinding);
+			return FECACommandResult::Error(FString::Printf(
+				TEXT("Conversion function '%s' is incompatible with source property '%s' or destination property '%s'."),
+				*ConvFunc->GetName(), *SourceProperty, *DestinationProperty));
+		}
+
+		FMVVMBlueprintFunctionReference FuncRef(WidgetBP, ConvFunc);
+		MVVMSubsystem->SetSourceToDestinationConversionFunction(WidgetBP, NewBinding, FuncRef);
+
+		// Wire the ViewModel property to the conversion function's input pin
+		UE::MVVM::FConversionFunctionValue ConvValue(ConvFunc);
+		TArray<FMVVMBlueprintPinId> InputPins;
+		if (VMProp)
+		{
+			InputPins = MVVMSubsystem->GetConversionFunctionPins(WidgetBP, ConvValue, VMProp);
+		}
+		if (InputPins.Num() > 0)
+		{
+			MVVMSubsystem->SetPathForConversionFunctionArgument(WidgetBP, NewBinding, InputPins[0], SourcePath, true);
+		}
+
+		ConversionFunctionUsed = ConvFunc->GetName();
+	}
+	else
+	{
+		// Direct binding — set source path directly
+		MVVMSubsystem->SetSourcePathForBinding(WidgetBP, NewBinding, SourcePath);
+	}
+
+	// Set destination path
+	MVVMSubsystem->SetDestinationPathForBinding(WidgetBP, NewBinding, DestPath, false);
+
+	// Capture binding ID before compile (reference may shift during compilation)
+	FGuid NewBindingId = NewBinding.BindingId;
+
+	CompileAndMarkDirty(WidgetBP);
+
+	// Post-wiring validation: verify the binding actually connected properly
+	const FMVVMBlueprintViewBinding* CompiledBinding = View->GetBinding(NewBindingId);
+	if (CompiledBinding)
+	{
+		FString DisplayStr = CompiledBinding->GetDisplayNameString(WidgetBP);
+
+		// Check for broken bindings — "<none>" means a path didn't resolve
+		if (DisplayStr.Contains(TEXT("<none>")))
+		{
+			// Collect compiler error messages if available
+			FString ErrorDetail;
+			TArray<FText> Errors = View->GetBindingMessages(NewBindingId, UE::MVVM::EBindingMessageType::Error);
+			for (const FText& Err : Errors)
+			{
+				if (!ErrorDetail.IsEmpty())
+				{
+					ErrorDetail += TEXT("; ");
+				}
+				ErrorDetail += Err.ToString();
+			}
+
+			// Remove the broken binding
+			MVVMSubsystem->RemoveBinding(WidgetBP, *CompiledBinding);
+			CompileAndMarkDirty(WidgetBP);
+
+			FString ErrorMsg = FString::Printf(
+				TEXT("Binding was created but failed validation. Display: '%s'."),
+				*DisplayStr);
+			if (!ErrorDetail.IsEmpty())
+			{
+				ErrorMsg += FString::Printf(TEXT(" Compiler errors: %s"), *ErrorDetail);
+			}
+			return FECACommandResult::Error(ErrorMsg);
+		}
+	}
+	else
+	{
+		return FECACommandResult::Error(
+			TEXT("Binding was created but could not be found after compilation. It may have been removed by the compiler."));
+	}
+
+	// Build result
+	TSharedPtr<FJsonObject> Result = MakeResult();
+	Result->SetStringField(TEXT("binding_mode"), BindingModeToString(BindingMode));
+	Result->SetStringField(TEXT("source"), FString::Printf(TEXT("%s.%s"), *VMContext->GetViewModelName().ToString(), *SourceProperty));
+	Result->SetStringField(TEXT("destination"), FString::Printf(TEXT("%s.%s"),
+		bIsSelfBinding ? TEXT("self") : *DestinationWidget, *DestinationProperty));
+	Result->SetStringField(TEXT("viewmodel_name"), VMContext->GetViewModelName().ToString());
+	if (CompiledBinding)
+	{
+		Result->SetStringField(TEXT("display"), CompiledBinding->GetDisplayNameString(WidgetBP));
+	}
+	if (!ConversionFunctionUsed.IsEmpty())
+	{
+		Result->SetStringField(TEXT("conversion_function"), ConversionFunctionUsed);
+		Result->SetStringField(TEXT("conversion_function_owner"), ConversionFunctionOwner);
+	}
+	return FECACommandResult::Success(Result);
+}
