ECABridge experimental MCP plugin

diff --git a/Engine/Plugins/Experimental/ECABridge/Source/ECABridge/Private/Commands/ECAMaterialNodeCommands.cpp b/Engine/Plugins/Experimental/ECABridge/Source/ECABridge/Private/Commands/ECAMaterialNodeCommands.cpp
new file mode 100644
--- /dev/null
+++ b/Engine/Plugins/Experimental/ECABridge/Source/ECABridge/Private/Commands/ECAMaterialNodeCommands.cpp
@@ -0,0 +1,3976 @@
+// Copyright Epic Games, Inc. All Rights Reserved.
+
+#include "Commands/ECAMaterialNodeCommands.h"
+#include "Commands/ECACommand.h"
+#include "Editor.h"
+#include "Toolkits/ToolkitManager.h"
+#include "IMaterialEditor.h"
+#include "Materials/Material.h"
+#include "Materials/MaterialExpressionAdd.h"
+#include "Materials/MaterialExpressionAppendVector.h"
+#include "Materials/MaterialExpressionClamp.h"
+#include "Materials/MaterialExpressionComponentMask.h"
+#include "Materials/MaterialExpressionConstant.h"
+#include "Materials/MaterialExpressionConstant2Vector.h"
+#include "Materials/MaterialExpressionConstant3Vector.h"
+#include "Materials/MaterialExpressionConstant4Vector.h"
+#include "Materials/MaterialExpressionDesaturation.h"
+#include "Materials/MaterialExpressionDivide.h"
+#include "Materials/MaterialExpressionDotProduct.h"
+#include "Materials/MaterialExpressionFresnel.h"
+#include "Materials/MaterialExpressionIf.h"
+#include "Materials/MaterialExpressionLinearInterpolate.h"
+#include "Materials/MaterialExpressionMultiply.h"
+#include "Materials/MaterialExpressionNormalize.h"
+#include "Materials/MaterialExpressionOneMinus.h"
+#include "Materials/MaterialExpressionPanner.h"
+#include "Materials/MaterialExpressionPower.h"
+#include "Materials/MaterialExpressionScalarParameter.h"
+#include "Materials/MaterialExpressionStaticBool.h"
+#include "Materials/MaterialExpressionStaticSwitch.h"
+#include "Materials/MaterialExpressionSubtract.h"
+#include "Materials/MaterialExpressionTextureSample.h"
+#include "Materials/MaterialExpressionTextureSampleParameter2D.h"
+#include "Materials/MaterialExpressionTextureCoordinate.h"
+#include "Materials/MaterialExpressionTextureObject.h"
+#include "Materials/MaterialExpressionTime.h"
+#include "Materials/MaterialExpressionVectorParameter.h"
+#include "Materials/MaterialExpressionWorldPosition.h"
+#include "Materials/MaterialExpressionCameraPositionWS.h"
+#include "Materials/MaterialExpressionVertexNormalWS.h"
+#include "Materials/MaterialExpressionPixelNormalWS.h"
+#include "Materials/MaterialExpressionAbs.h"
+#include "Materials/MaterialExpressionCeil.h"
+#include "Materials/MaterialExpressionFloor.h"
+#include "Materials/MaterialExpressionFrac.h"
+#include "Materials/MaterialExpressionSaturate.h"
+#include "Materials/MaterialExpressionSine.h"
+#include "Materials/MaterialExpressionCosine.h"
+#include "Materials/MaterialExpressionSquareRoot.h"
+#include "Materials/MaterialExpressionDistance.h"
+#include "Materials/MaterialExpressionCrossProduct.h"
+#include "Materials/MaterialExpressionMakeMaterialAttributes.h"
+#include "Materials/MaterialExpressionBreakMaterialAttributes.h"
+#include "Materials/MaterialExpressionComment.h"
+#include "Materials/MaterialExpressionNoise.h"
+#include "Materials/MaterialExpressionCustom.h"
+#include "Engine/Texture2D.h"
+#include "UObject/UObjectIterator.h"
+#include "EdGraph/EdGraphNode.h"
+#include "Math/Color.h"
+#include "Math/Vector.h"
+#include "Math/Vector2D.h"
+#include "Math/Vector4.h"
+#include "Math/Rotator.h"
+#include "Math/Transform.h"
+
+// Register all material node commands
+REGISTER_ECA_COMMAND(FECACommand_AddMaterialNode)
+REGISTER_ECA_COMMAND(FECACommand_ConnectMaterialNodes)
+REGISTER_ECA_COMMAND(FECACommand_GetMaterialNodes)
+REGISTER_ECA_COMMAND(FECACommand_GetMaterialNodeInfo)
+REGISTER_ECA_COMMAND(FECACommand_DeleteMaterialNode)
+REGISTER_ECA_COMMAND(FECACommand_DisconnectMaterialNode)
+REGISTER_ECA_COMMAND(FECACommand_SetMaterialNodeProperty)
+REGISTER_ECA_COMMAND(FECACommand_BatchEditMaterialNodes)
+REGISTER_ECA_COMMAND(FECACommand_ListMaterialExpressionTypes)
+REGISTER_ECA_COMMAND(FECACommand_GetMaterialErrors)
+REGISTER_ECA_COMMAND(FECACommand_SetCustomNodeInputName)
+REGISTER_ECA_COMMAND(FECACommand_SetMaterialProperty)
+REGISTER_ECA_COMMAND(FECACommand_GetMaterialProperty)
+REGISTER_ECA_COMMAND(FECACommand_AutoLayoutMaterialGraph)
+REGISTER_ECA_COMMAND(FECACommand_RenameParameterGroup)
+REGISTER_ECA_COMMAND(FECACommand_ListParameterGroups)
+
+//------------------------------------------------------------------------------
+// Helper Functions
+//------------------------------------------------------------------------------
+
+// Collect compilation errors from material expressions
+static TArray<TSharedPtr<FJsonValue>> GetMaterialCompilationErrors(UMaterial* Material)
+{
+	TArray<TSharedPtr<FJsonValue>> ErrorsArray;
+	
+	if (!Material || !Material->GetEditorOnlyData())
+	{
+		return ErrorsArray;
+	}
+	
+	for (UMaterialExpression* Expr : Material->GetEditorOnlyData()->ExpressionCollection.Expressions)
+	{
+		if (Expr && !Expr->LastErrorText.IsEmpty())
+		{
+			TSharedPtr<FJsonObject> ErrorObj = MakeShared<FJsonObject>();
+			ErrorObj->SetStringField(TEXT("node_id"), Expr->MaterialExpressionGuid.ToString());
+			ErrorObj->SetStringField(TEXT("node_type"), Expr->GetClass()->GetName().Replace(TEXT("MaterialExpression"), TEXT("")));
+			ErrorObj->SetStringField(TEXT("error"), Expr->LastErrorText);
+			ErrorsArray.Add(MakeShared<FJsonValueObject>(ErrorObj));
+		}
+	}
+	
+	return ErrorsArray;
+}
+
+// Clear previous compilation errors from all expressions before recompiling
+static void ClearMaterialCompilationErrors(UMaterial* Material)
+{
+	if (!Material || !Material->GetEditorOnlyData())
+	{
+		return;
+	}
+	
+	for (UMaterialExpression* Expr : Material->GetEditorOnlyData()->ExpressionCollection.Expressions)
+	{
+		if (Expr)
+		{
+			Expr->LastErrorText.Empty();
+		}
+	}
+}
+
+// Add compilation errors to result if any exist
+static void AddCompilationErrorsToResult(TSharedPtr<FJsonObject>& Result, UMaterial* Material)
+{
+	TArray<TSharedPtr<FJsonValue>> Errors = GetMaterialCompilationErrors(Material);
+	if (Errors.Num() > 0)
+	{
+		Result->SetArrayField(TEXT("compilation_errors"), Errors);
+		Result->SetBoolField(TEXT("has_errors"), true);
+	}
+	else
+	{
+		Result->SetBoolField(TEXT("has_errors"), false);
+	}
+}
+
+// Refresh the material editor UI if the material is currently open
+static void RefreshMaterialEditorUI(UMaterial* Material)
+{
+	if (!Material)
+	{
+		return;
+	}
+	
+	// Use FToolkitManager to find the editor - this is the same approach used by FMaterialEditorUtilities
+	TSharedPtr<IToolkit> FoundAssetEditor = FToolkitManager::Get().FindEditorForAsset(Material);
+	if (FoundAssetEditor.IsValid())
+	{
+		TSharedPtr<IMaterialEditor> MaterialEditor = StaticCastSharedPtr<IMaterialEditor>(FoundAssetEditor);
+		if (MaterialEditor.IsValid())
+		{
+			MaterialEditor->UpdateMaterialAfterGraphChange();
+		}
+	}
+}
+
+static UMaterial* LoadMaterialByPath(const FString& MaterialPath)
+{
+	return LoadObject<UMaterial>(nullptr, *MaterialPath);
+}
+
+// Constants for auto-layout spacing in materials
+static const int32 MAT_AUTO_LAYOUT_SPACING_X = 300;
+static const int32 MAT_AUTO_LAYOUT_SPACING_Y = 150;
+
+// Calculate automatic position for a new material expression node
+// Places the node to the left of existing nodes (material flow goes right-to-left)
+static FVector2D CalculateAutoMaterialNodePosition(UMaterial* Material)
+{
+	if (!Material || !Material->GetEditorOnlyData())
+	{
+		return FVector2D(0, 0);
+	}
+	
+	const auto& Expressions = Material->GetEditorOnlyData()->ExpressionCollection.Expressions;
+	
+	// Find the bounds of existing nodes
+	int32 MinX = INT_MAX;
+	int32 MaxX = INT_MIN;
+	int32 MinY = INT_MAX;
+	int32 MaxY = INT_MIN;
+	int32 NodeCount = 0;
+	
+	// Track the leftmost node at each Y level (materials flow right-to-left)
+	TMap<int32, int32> LeftmostXAtY; // Y bucket -> min X
+	
+	for (UMaterialExpression* Expr : Expressions)
+	{
+		if (Expr)
+		{
+			MinX = FMath::Min(MinX, Expr->MaterialExpressionEditorX);
+			MaxX = FMath::Max(MaxX, Expr->MaterialExpressionEditorX);
+			MinY = FMath::Min(MinY, Expr->MaterialExpressionEditorY);
+			MaxY = FMath::Max(MaxY, Expr->MaterialExpressionEditorY);
+			NodeCount++;
+			
+			// Track leftmost node in Y buckets (rounded to spacing)
+			int32 YBucket = (Expr->MaterialExpressionEditorY / MAT_AUTO_LAYOUT_SPACING_Y) * MAT_AUTO_LAYOUT_SPACING_Y;
+			if (!LeftmostXAtY.Contains(YBucket) || Expr->MaterialExpressionEditorX < LeftmostXAtY[YBucket])
+			{
+				LeftmostXAtY.Add(YBucket, Expr->MaterialExpressionEditorX);
+			}
+		}
+	}
+	
+	// If no nodes exist, start at a position left of the material result
+	if (NodeCount == 0)
+	{
+		return FVector2D(-MAT_AUTO_LAYOUT_SPACING_X, 0);
+	}
+	
+	// Material graphs flow right-to-left, so place new nodes to the LEFT
+	int32 BestY = MinY;
+	int32 BestX = MinX - MAT_AUTO_LAYOUT_SPACING_X;
+	
+	// If the graph is getting wide, find the Y level with most room on the left
+	if (MaxX - MinX > MAT_AUTO_LAYOUT_SPACING_X * 5)
+	{
+		int32 MaxLeftX = INT_MIN;
+		for (const auto& Pair : LeftmostXAtY)
+		{
+			if (Pair.Value > MaxLeftX)
+			{
+				MaxLeftX = Pair.Value;
+				BestY = Pair.Key;
+				BestX = MaxLeftX - MAT_AUTO_LAYOUT_SPACING_X;
+			}
+		}
+		
+		// Or start a completely new row below everything
+		if (BestX < MinX - MAT_AUTO_LAYOUT_SPACING_X * 2)
+		{
+			BestY = MaxY + MAT_AUTO_LAYOUT_SPACING_Y;
+			BestX = MinX - MAT_AUTO_LAYOUT_SPACING_X;
+		}
+	}
+	
+	return FVector2D(BestX, BestY);
+}
+
+static FVector2D GetNodePosition(const TSharedPtr<FJsonObject>& Params, UMaterial* Material = nullptr)
+{
+	FVector2D Position(0, 0);
+	const TSharedPtr<FJsonObject>* PosObj;
+	
+	// Check if explicit position was provided
+	if (Params->TryGetObjectField(TEXT("node_position"), PosObj))
+	{
+		Position.X = (*PosObj)->GetNumberField(TEXT("x"));
+		Position.Y = (*PosObj)->GetNumberField(TEXT("y"));
+	}
+	else if (Material)
+	{
+		// No position provided - calculate automatically
+		Position = CalculateAutoMaterialNodePosition(Material);
+	}
+	
+	return Position;
+}
+
+static UMaterialExpression* FindExpressionByGuid(UMaterial* Material, const FGuid& Guid)
+{
+	for (UMaterialExpression* Expr : Material->GetEditorOnlyData()->ExpressionCollection.Expressions)
+	{
+		if (Expr && Expr->MaterialExpressionGuid == Guid)
+		{
+			return Expr;
+		}
+	}
+	return nullptr;
+}
+
+static TSharedPtr<FJsonObject> ExpressionToJson(UMaterialExpression* Expression)
+{
+	TSharedPtr<FJsonObject> ExprJson = MakeShared<FJsonObject>();
+	
+	ExprJson->SetStringField(TEXT("node_id"), Expression->MaterialExpressionGuid.ToString());
+	ExprJson->SetStringField(TEXT("node_class"), Expression->GetClass()->GetName());
+	ExprJson->SetStringField(TEXT("node_type"), Expression->GetClass()->GetName().Replace(TEXT("MaterialExpression"), TEXT("")));
+	ExprJson->SetNumberField(TEXT("x"), Expression->MaterialExpressionEditorX);
+	ExprJson->SetNumberField(TEXT("y"), Expression->MaterialExpressionEditorY);
+	
+	if (!Expression->Desc.IsEmpty())
+	{
+		ExprJson->SetStringField(TEXT("description"), Expression->Desc);
+	}
+	
+	// Add outputs info
+	TArray<TSharedPtr<FJsonValue>> OutputsArray;
+	for (int32 i = 0; i < Expression->Outputs.Num(); i++)
+	{
+		TSharedPtr<FJsonObject> OutputJson = MakeShared<FJsonObject>();
+		OutputJson->SetStringField(TEXT("name"), Expression->Outputs[i].OutputName.IsNone() ? FString::FromInt(i) : Expression->Outputs[i].OutputName.ToString());
+		OutputJson->SetNumberField(TEXT("index"), i);
+		OutputsArray.Add(MakeShared<FJsonValueObject>(OutputJson));
+	}
+	ExprJson->SetArrayField(TEXT("outputs"), OutputsArray);
+	
+	// Add inputs info
+	TArray<TSharedPtr<FJsonValue>> InputsArray;
+	for (int32 i = 0; ; i++)
+	{
+		FExpressionInput* Input = Expression->GetInput(i);
+		if (!Input)
+		{
+			break;
+		}
+		
+		TSharedPtr<FJsonObject> InputJson = MakeShared<FJsonObject>();
+		FName InputName = Expression->GetInputName(i);
+		InputJson->SetStringField(TEXT("name"), InputName.IsNone() ? FString::FromInt(i) : InputName.ToString());
+		InputJson->SetNumberField(TEXT("index"), i);
+		InputJson->SetBoolField(TEXT("is_connected"), Input->Expression != nullptr);
+		
+		if (Input->Expression)
+		{
+			InputJson->SetStringField(TEXT("connected_node_id"), Input->Expression->MaterialExpressionGuid.ToString());
+			InputJson->SetNumberField(TEXT("connected_output_index"), Input->OutputIndex);
+		}
+		
+		InputsArray.Add(MakeShared<FJsonValueObject>(InputJson));
+	}
+	ExprJson->SetArrayField(TEXT("inputs"), InputsArray);
+	
+	// Add parameter-specific info
+	if (UMaterialExpressionScalarParameter* ScalarParam = Cast<UMaterialExpressionScalarParameter>(Expression))
+	{
+		ExprJson->SetStringField(TEXT("parameter_name"), ScalarParam->ParameterName.ToString());
+		ExprJson->SetNumberField(TEXT("default_value"), ScalarParam->DefaultValue);
+	}
+	else if (UMaterialExpressionVectorParameter* VectorParam = Cast<UMaterialExpressionVectorParameter>(Expression))
+	{
+		ExprJson->SetStringField(TEXT("parameter_name"), VectorParam->ParameterName.ToString());
+		TSharedPtr<FJsonObject> ColorJson = MakeShared<FJsonObject>();
+		ColorJson->SetNumberField(TEXT("r"), VectorParam->DefaultValue.R);
+		ColorJson->SetNumberField(TEXT("g"), VectorParam->DefaultValue.G);
+		ColorJson->SetNumberField(TEXT("b"), VectorParam->DefaultValue.B);
+		ColorJson->SetNumberField(TEXT("a"), VectorParam->DefaultValue.A);
+		ExprJson->SetObjectField(TEXT("default_value"), ColorJson);
+	}
+	else if (UMaterialExpressionTextureSampleParameter2D* TexParam = Cast<UMaterialExpressionTextureSampleParameter2D>(Expression))
+	{
+		ExprJson->SetStringField(TEXT("parameter_name"), TexParam->ParameterName.ToString());
+		if (TexParam->Texture)
+		{
+			ExprJson->SetStringField(TEXT("texture_path"), TexParam->Texture->GetPathName());
+		}
+	}
+	else if (UMaterialExpressionTextureSample* TexSample = Cast<UMaterialExpressionTextureSample>(Expression))
+	{
+		if (TexSample->Texture)
+		{
+			ExprJson->SetStringField(TEXT("texture_path"), TexSample->Texture->GetPathName());
+		}
+		ExprJson->SetNumberField(TEXT("sampler_type"), static_cast<int32>(TexSample->SamplerType));
+		ExprJson->SetNumberField(TEXT("sampler_source"), static_cast<int32>(TexSample->SamplerSource));
+		ExprJson->SetNumberField(TEXT("mip_value_mode"), static_cast<int32>(TexSample->MipValueMode));
+		ExprJson->SetNumberField(TEXT("const_coordinate"), TexSample->ConstCoordinate);
+		ExprJson->SetNumberField(TEXT("const_mip_value"), TexSample->ConstMipValue);
+		ExprJson->SetBoolField(TEXT("automatic_view_mip_bias"), TexSample->AutomaticViewMipBias);
+	}
+	else if (UMaterialExpressionConstant* Const = Cast<UMaterialExpressionConstant>(Expression))
+	{
+		ExprJson->SetNumberField(TEXT("value"), Const->R);
+	}
+	else if (UMaterialExpressionConstant2Vector* Const2 = Cast<UMaterialExpressionConstant2Vector>(Expression))
+	{
+		TSharedPtr<FJsonObject> VecJson = MakeShared<FJsonObject>();
+		VecJson->SetNumberField(TEXT("r"), Const2->R);
+		VecJson->SetNumberField(TEXT("g"), Const2->G);
+		ExprJson->SetObjectField(TEXT("value"), VecJson);
+	}
+	else if (UMaterialExpressionConstant3Vector* Const3 = Cast<UMaterialExpressionConstant3Vector>(Expression))
+	{
+		TSharedPtr<FJsonObject> ColorJson = MakeShared<FJsonObject>();
+		ColorJson->SetNumberField(TEXT("r"), Const3->Constant.R);
+		ColorJson->SetNumberField(TEXT("g"), Const3->Constant.G);
+		ColorJson->SetNumberField(TEXT("b"), Const3->Constant.B);
+		ExprJson->SetObjectField(TEXT("value"), ColorJson);
+	}
+	else if (UMaterialExpressionConstant4Vector* Const4 = Cast<UMaterialExpressionConstant4Vector>(Expression))
+	{
+		TSharedPtr<FJsonObject> ColorJson = MakeShared<FJsonObject>();
+		ColorJson->SetNumberField(TEXT("r"), Const4->Constant.R);
+		ColorJson->SetNumberField(TEXT("g"), Const4->Constant.G);
+		ColorJson->SetNumberField(TEXT("b"), Const4->Constant.B);
+		ColorJson->SetNumberField(TEXT("a"), Const4->Constant.A);
+		ExprJson->SetObjectField(TEXT("value"), ColorJson);
+	}
+	else if (UMaterialExpressionTextureCoordinate* TexCoord = Cast<UMaterialExpressionTextureCoordinate>(Expression))
+	{
+		ExprJson->SetNumberField(TEXT("coordinate_index"), TexCoord->CoordinateIndex);
+		ExprJson->SetNumberField(TEXT("u_tiling"), TexCoord->UTiling);
+		ExprJson->SetNumberField(TEXT("v_tiling"), TexCoord->VTiling);
+	}
+	else if (UMaterialExpressionComponentMask* Mask = Cast<UMaterialExpressionComponentMask>(Expression))
+	{
+		ExprJson->SetBoolField(TEXT("r"), Mask->R);
+		ExprJson->SetBoolField(TEXT("g"), Mask->G);
+		ExprJson->SetBoolField(TEXT("b"), Mask->B);
+		ExprJson->SetBoolField(TEXT("a"), Mask->A);
+	}
+	else if (UMaterialExpressionNoise* Noise = Cast<UMaterialExpressionNoise>(Expression))
+	{
+		ExprJson->SetNumberField(TEXT("scale"), Noise->Scale);
+		ExprJson->SetNumberField(TEXT("quality"), Noise->Quality);
+		ExprJson->SetNumberField(TEXT("noise_function"), static_cast<int32>(Noise->NoiseFunction));
+		ExprJson->SetBoolField(TEXT("turbulence"), Noise->bTurbulence);
+		ExprJson->SetNumberField(TEXT("levels"), Noise->Levels);
+		ExprJson->SetNumberField(TEXT("output_min"), Noise->OutputMin);
+		ExprJson->SetNumberField(TEXT("output_max"), Noise->OutputMax);
+		ExprJson->SetNumberField(TEXT("level_scale"), Noise->LevelScale);
+		ExprJson->SetBoolField(TEXT("tiling"), Noise->bTiling);
+		ExprJson->SetNumberField(TEXT("repeat_size"), Noise->RepeatSize);
+	}
+	else if (UMaterialExpressionCustom* Custom = Cast<UMaterialExpressionCustom>(Expression))
+	{
+		ExprJson->SetStringField(TEXT("code"), Custom->Code);
+		ExprJson->SetNumberField(TEXT("output_type"), static_cast<int32>(Custom->OutputType));
+		ExprJson->SetStringField(TEXT("description"), Custom->Description);
+		
+		// Serialize inputs
+		TArray<TSharedPtr<FJsonValue>> CustomInputsArray;
+		for (const FCustomInput& Input : Custom->Inputs)
+		{
+			TSharedPtr<FJsonObject> InputObj = MakeShared<FJsonObject>();
+			InputObj->SetStringField(TEXT("input_name"), Input.InputName.ToString());
+			if (Input.Input.Expression)
+			{
+				InputObj->SetStringField(TEXT("connected_node_id"), Input.Input.Expression->MaterialExpressionGuid.ToString());
+				InputObj->SetNumberField(TEXT("connected_output_index"), Input.Input.OutputIndex);
+			}
+			CustomInputsArray.Add(MakeShared<FJsonValueObject>(InputObj));
+		}
+		ExprJson->SetArrayField(TEXT("custom_inputs"), CustomInputsArray);
+		
+		// Serialize additional outputs
+		TArray<TSharedPtr<FJsonValue>> CustomOutputsArray;
+		for (const FCustomOutput& Output : Custom->AdditionalOutputs)
+		{
+			TSharedPtr<FJsonObject> OutputObj = MakeShared<FJsonObject>();
+			OutputObj->SetStringField(TEXT("output_name"), Output.OutputName.ToString());
+			OutputObj->SetNumberField(TEXT("output_type"), static_cast<int32>(Output.OutputType));
+			CustomOutputsArray.Add(MakeShared<FJsonValueObject>(OutputObj));
+		}
+		ExprJson->SetArrayField(TEXT("additional_outputs"), CustomOutputsArray);
+		
+		// Serialize defines
+		TArray<TSharedPtr<FJsonValue>> DefinesArray;
+		for (const FCustomDefine& Define : Custom->AdditionalDefines)
+		{
+			TSharedPtr<FJsonObject> DefineObj = MakeShared<FJsonObject>();
+			DefineObj->SetStringField(TEXT("define_name"), Define.DefineName);
+			DefineObj->SetStringField(TEXT("define_value"), Define.DefineValue);
+			DefinesArray.Add(MakeShared<FJsonValueObject>(DefineObj));
+		}
+		ExprJson->SetArrayField(TEXT("additional_defines"), DefinesArray);
+		
+		// Serialize include paths
+		TArray<TSharedPtr<FJsonValue>> IncludesArray;
+		for (const FString& IncludePath : Custom->IncludeFilePaths)
+		{
+			IncludesArray.Add(MakeShared<FJsonValueString>(IncludePath));
+		}
+		ExprJson->SetArrayField(TEXT("include_file_paths"), IncludesArray);
+	}
+	
+	return ExprJson;
+}
+
+// Map of simple node type names to their UClass
+static TMap<FString, UClass*> GetExpressionClassMap()
+{
+	static TMap<FString, UClass*> ClassMap;
+	if (ClassMap.Num() == 0)
+	{
+		// Math operations
+		ClassMap.Add(TEXT("Add"), UMaterialExpressionAdd::StaticClass());
+		ClassMap.Add(TEXT("Subtract"), UMaterialExpressionSubtract::StaticClass());
+		ClassMap.Add(TEXT("Multiply"), UMaterialExpressionMultiply::StaticClass());
+		ClassMap.Add(TEXT("Divide"), UMaterialExpressionDivide::StaticClass());
+		ClassMap.Add(TEXT("Power"), UMaterialExpressionPower::StaticClass());
+		ClassMap.Add(TEXT("SquareRoot"), UMaterialExpressionSquareRoot::StaticClass());
+		ClassMap.Add(TEXT("Abs"), UMaterialExpressionAbs::StaticClass());
+		ClassMap.Add(TEXT("Ceil"), UMaterialExpressionCeil::StaticClass());
+		ClassMap.Add(TEXT("Floor"), UMaterialExpressionFloor::StaticClass());
+		ClassMap.Add(TEXT("Frac"), UMaterialExpressionFrac::StaticClass());
+		ClassMap.Add(TEXT("Clamp"), UMaterialExpressionClamp::StaticClass());
+		ClassMap.Add(TEXT("Saturate"), UMaterialExpressionSaturate::StaticClass());
+		ClassMap.Add(TEXT("OneMinus"), UMaterialExpressionOneMinus::StaticClass());
+		ClassMap.Add(TEXT("Sine"), UMaterialExpressionSine::StaticClass());
+		ClassMap.Add(TEXT("Cosine"), UMaterialExpressionCosine::StaticClass());
+		
+		// Vector operations
+		ClassMap.Add(TEXT("DotProduct"), UMaterialExpressionDotProduct::StaticClass());
+		ClassMap.Add(TEXT("CrossProduct"), UMaterialExpressionCrossProduct::StaticClass());
+		ClassMap.Add(TEXT("Normalize"), UMaterialExpressionNormalize::StaticClass());
+		ClassMap.Add(TEXT("Distance"), UMaterialExpressionDistance::StaticClass());
+		ClassMap.Add(TEXT("ComponentMask"), UMaterialExpressionComponentMask::StaticClass());
+		ClassMap.Add(TEXT("AppendVector"), UMaterialExpressionAppendVector::StaticClass());
+		ClassMap.Add(TEXT("Append"), UMaterialExpressionAppendVector::StaticClass());
+		
+		// Interpolation
+		ClassMap.Add(TEXT("Lerp"), UMaterialExpressionLinearInterpolate::StaticClass());
+		ClassMap.Add(TEXT("LinearInterpolate"), UMaterialExpressionLinearInterpolate::StaticClass());
+		ClassMap.Add(TEXT("If"), UMaterialExpressionIf::StaticClass());
+		ClassMap.Add(TEXT("StaticSwitch"), UMaterialExpressionStaticSwitch::StaticClass());
+		
+		// Constants
+		ClassMap.Add(TEXT("Constant"), UMaterialExpressionConstant::StaticClass());
+		ClassMap.Add(TEXT("Constant2Vector"), UMaterialExpressionConstant2Vector::StaticClass());
+		ClassMap.Add(TEXT("Constant3Vector"), UMaterialExpressionConstant3Vector::StaticClass());
+		ClassMap.Add(TEXT("Constant4Vector"), UMaterialExpressionConstant4Vector::StaticClass());
+		
+		// Parameters
+		ClassMap.Add(TEXT("ScalarParameter"), UMaterialExpressionScalarParameter::StaticClass());
+		ClassMap.Add(TEXT("VectorParameter"), UMaterialExpressionVectorParameter::StaticClass());
+		ClassMap.Add(TEXT("StaticBool"), UMaterialExpressionStaticBool::StaticClass());
+		
+		// Textures
+		ClassMap.Add(TEXT("TextureSample"), UMaterialExpressionTextureSample::StaticClass());
+		ClassMap.Add(TEXT("TextureSampleParameter2D"), UMaterialExpressionTextureSampleParameter2D::StaticClass());
+		ClassMap.Add(TEXT("TextureCoordinate"), UMaterialExpressionTextureCoordinate::StaticClass());
+		ClassMap.Add(TEXT("TexCoord"), UMaterialExpressionTextureCoordinate::StaticClass());
+		ClassMap.Add(TEXT("TextureObject"), UMaterialExpressionTextureObject::StaticClass());
+		ClassMap.Add(TEXT("Panner"), UMaterialExpressionPanner::StaticClass());
+		
+		// World/Camera
+		ClassMap.Add(TEXT("WorldPosition"), UMaterialExpressionWorldPosition::StaticClass());
+		ClassMap.Add(TEXT("CameraPosition"), UMaterialExpressionCameraPositionWS::StaticClass());
+		ClassMap.Add(TEXT("VertexNormalWS"), UMaterialExpressionVertexNormalWS::StaticClass());
+		ClassMap.Add(TEXT("PixelNormalWS"), UMaterialExpressionPixelNormalWS::StaticClass());
+		
+		// Utility
+		ClassMap.Add(TEXT("Time"), UMaterialExpressionTime::StaticClass());
+		ClassMap.Add(TEXT("Desaturation"), UMaterialExpressionDesaturation::StaticClass());
+		ClassMap.Add(TEXT("Fresnel"), UMaterialExpressionFresnel::StaticClass());
+		
+		// Material Attributes
+		ClassMap.Add(TEXT("MakeMaterialAttributes"), UMaterialExpressionMakeMaterialAttributes::StaticClass());
+		ClassMap.Add(TEXT("BreakMaterialAttributes"), UMaterialExpressionBreakMaterialAttributes::StaticClass());
+		
+		// Comments
+		ClassMap.Add(TEXT("Comment"), UMaterialExpressionComment::StaticClass());
+		
+		// Procedural
+		ClassMap.Add(TEXT("Noise"), UMaterialExpressionNoise::StaticClass());
+		
+		// Custom code
+		ClassMap.Add(TEXT("Custom"), UMaterialExpressionCustom::StaticClass());
+	}
+	return ClassMap;
+}
+
+static UClass* FindExpressionClass(const FString& NodeType)
+{
+	// First check our simple name map
+	TMap<FString, UClass*> ClassMap = GetExpressionClassMap();
+	if (UClass** Found = ClassMap.Find(NodeType))
+	{
+		return *Found;
+	}
+	
+	// Try with MaterialExpression prefix
+	FString FullClassName = TEXT("MaterialExpression") + NodeType;
+	for (TObjectIterator<UClass> It; It; ++It)
+	{
+		if (It->IsChildOf(UMaterialExpression::StaticClass()) && !It->HasAnyClassFlags(CLASS_Abstract))
+		{
+			if (It->GetName().Equals(FullClassName, ESearchCase::IgnoreCase) ||
+				It->GetName().Equals(NodeType, ESearchCase::IgnoreCase))
+			{
+				return *It;
+			}
+		}
+	}
+	
+	return nullptr;
+}
+
+static FExpressionInput* GetMaterialInput(UMaterial* Material, const FString& InputName)
+{
+	// Map common input names to material properties
+	if (InputName.Equals(TEXT("BaseColor"), ESearchCase::IgnoreCase))
+		return &Material->GetEditorOnlyData()->BaseColor;
+	if (InputName.Equals(TEXT("Metallic"), ESearchCase::IgnoreCase))
+		return &Material->GetEditorOnlyData()->Metallic;
+	if (InputName.Equals(TEXT("Specular"), ESearchCase::IgnoreCase))
+		return &Material->GetEditorOnlyData()->Specular;
+	if (InputName.Equals(TEXT("Roughness"), ESearchCase::IgnoreCase))
+		return &Material->GetEditorOnlyData()->Roughness;
+	if (InputName.Equals(TEXT("Anisotropy"), ESearchCase::IgnoreCase))
+		return &Material->GetEditorOnlyData()->Anisotropy;
+	if (InputName.Equals(TEXT("Normal"), ESearchCase::IgnoreCase))
+		return &Material->GetEditorOnlyData()->Normal;
+	if (InputName.Equals(TEXT("Tangent"), ESearchCase::IgnoreCase))
+		return &Material->GetEditorOnlyData()->Tangent;
+	if (InputName.Equals(TEXT("EmissiveColor"), ESearchCase::IgnoreCase) || InputName.Equals(TEXT("Emissive"), ESearchCase::IgnoreCase))
+		return &Material->GetEditorOnlyData()->EmissiveColor;
+	if (InputName.Equals(TEXT("Opacity"), ESearchCase::IgnoreCase))
+		return &Material->GetEditorOnlyData()->Opacity;
+	if (InputName.Equals(TEXT("OpacityMask"), ESearchCase::IgnoreCase))
+		return &Material->GetEditorOnlyData()->OpacityMask;
+	if (InputName.Equals(TEXT("WorldPositionOffset"), ESearchCase::IgnoreCase) || InputName.Equals(TEXT("WPO"), ESearchCase::IgnoreCase))
+		return &Material->GetEditorOnlyData()->WorldPositionOffset;
+	if (InputName.Equals(TEXT("SubsurfaceColor"), ESearchCase::IgnoreCase))
+		return &Material->GetEditorOnlyData()->SubsurfaceColor;
+	if (InputName.Equals(TEXT("AmbientOcclusion"), ESearchCase::IgnoreCase) || InputName.Equals(TEXT("AO"), ESearchCase::IgnoreCase))
+		return &Material->GetEditorOnlyData()->AmbientOcclusion;
+	if (InputName.Equals(TEXT("Refraction"), ESearchCase::IgnoreCase))
+		return &Material->GetEditorOnlyData()->Refraction;
+	if (InputName.Equals(TEXT("PixelDepthOffset"), ESearchCase::IgnoreCase))
+		return &Material->GetEditorOnlyData()->PixelDepthOffset;
+	if (InputName.Equals(TEXT("Displacement"), ESearchCase::IgnoreCase))
+		return &Material->GetEditorOnlyData()->Displacement;
+	if (InputName.Equals(TEXT("ShadingModel"), ESearchCase::IgnoreCase) || InputName.Equals(TEXT("ShadingModelFromMaterialExpression"), ESearchCase::IgnoreCase))
+		return &Material->GetEditorOnlyData()->ShadingModelFromMaterialExpression;
+	if (InputName.Equals(TEXT("FrontMaterial"), ESearchCase::IgnoreCase))
+		return &Material->GetEditorOnlyData()->FrontMaterial;
+	if (InputName.Equals(TEXT("SurfaceThickness"), ESearchCase::IgnoreCase))
+		return &Material->GetEditorOnlyData()->SurfaceThickness;
+	if (InputName.Equals(TEXT("ClearCoat"), ESearchCase::IgnoreCase))
+		return &Material->GetEditorOnlyData()->ClearCoat;
+	if (InputName.Equals(TEXT("ClearCoatRoughness"), ESearchCase::IgnoreCase))
+		return &Material->GetEditorOnlyData()->ClearCoatRoughness;
+	
+	return nullptr;
+}
+
+static int32 ParseOutputIndex(const FString& OutputStr, UMaterialExpression* Expression)
+{
+	// Try parsing as number first
+	if (OutputStr.IsNumeric())
+	{
+		return FCString::Atoi(*OutputStr);
+	}
+	
+	// Map common output names
+	if (OutputStr.Equals(TEXT("RGB"), ESearchCase::IgnoreCase))
+		return 0;
+	if (OutputStr.Equals(TEXT("R"), ESearchCase::IgnoreCase))
+		return 1;
+	if (OutputStr.Equals(TEXT("G"), ESearchCase::IgnoreCase))
+		return 2;
+	if (OutputStr.Equals(TEXT("B"), ESearchCase::IgnoreCase))
+		return 3;
+	if (OutputStr.Equals(TEXT("A"), ESearchCase::IgnoreCase))
+		return 4;
+	
+	// Try to find by output name in the expression
+	for (int32 i = 0; i < Expression->Outputs.Num(); i++)
+	{
+		if (Expression->Outputs[i].OutputName.ToString().Equals(OutputStr, ESearchCase::IgnoreCase))
+		{
+			return i;
+		}
+	}
+	
+	return 0; // Default to first output
+}
+
+//------------------------------------------------------------------------------
+// AddMaterialNode
+//------------------------------------------------------------------------------
+
+FECACommandResult FECACommand_AddMaterialNode::Execute(const TSharedPtr<FJsonObject>& Params)
+{
+	FString MaterialPath;
+	if (!GetStringParam(Params, TEXT("material_path"), MaterialPath))
+	{
+		return FECACommandResult::Error(TEXT("Missing required parameter: material_path"));
+	}
+	
+	FString NodeType;
+	if (!GetStringParam(Params, TEXT("node_type"), NodeType))
+	{
+		return FECACommandResult::Error(TEXT("Missing required parameter: node_type"));
+	}
+	
+	UMaterial* Material = LoadMaterialByPath(MaterialPath);
+	if (!Material)
+	{
+		return FECACommandResult::Error(FString::Printf(TEXT("Material not found: %s"), *MaterialPath));
+	}
+	
+	UClass* ExpressionClass = FindExpressionClass(NodeType);
+	if (!ExpressionClass)
+	{
+		return FECACommandResult::Error(FString::Printf(TEXT("Unknown material expression type: %s"), *NodeType));
+	}
+	
+	FVector2D Position = GetNodePosition(Params, Material);
+	
+	// Create the expression
+	Material->PreEditChange(nullptr);
+	
+	UMaterialExpression* NewExpression = NewObject<UMaterialExpression>(Material, ExpressionClass, NAME_None, RF_Transactional);
+	NewExpression->MaterialExpressionEditorX = Position.X;
+	NewExpression->MaterialExpressionEditorY = Position.Y;
+	
+	// Configure based on type
+	FString ParameterName;
+	if (GetStringParam(Params, TEXT("parameter_name"), ParameterName, false))
+	{
+		if (UMaterialExpressionScalarParameter* ScalarParam = Cast<UMaterialExpressionScalarParameter>(NewExpression))
+		{
+			ScalarParam->ParameterName = FName(*ParameterName);
+		}
+		else if (UMaterialExpressionVectorParameter* VectorParam = Cast<UMaterialExpressionVectorParameter>(NewExpression))
+		{
+			VectorParam->ParameterName = FName(*ParameterName);
+		}
+		else if (UMaterialExpressionTextureSampleParameter2D* TexParam = Cast<UMaterialExpressionTextureSampleParameter2D>(NewExpression))
+		{
+			TexParam->ParameterName = FName(*ParameterName);
+		}
+	}
+	
+	// Set default value
+	const TSharedPtr<FJsonObject>* DefaultValueObj;
+	double DefaultValueNum;
+	FString DefaultValueStr;
+	
+	if (Params->TryGetObjectField(TEXT("default_value"), DefaultValueObj))
+	{
+		// Vector/Color value
+		if (UMaterialExpressionVectorParameter* VectorParam = Cast<UMaterialExpressionVectorParameter>(NewExpression))
+		{
+			VectorParam->DefaultValue.R = (*DefaultValueObj)->GetNumberField(TEXT("r"));
+			VectorParam->DefaultValue.G = (*DefaultValueObj)->GetNumberField(TEXT("g"));
+			VectorParam->DefaultValue.B = (*DefaultValueObj)->GetNumberField(TEXT("b"));
+			if ((*DefaultValueObj)->HasField(TEXT("a")))
+				VectorParam->DefaultValue.A = (*DefaultValueObj)->GetNumberField(TEXT("a"));
+		}
+		else if (UMaterialExpressionConstant2Vector* Const2 = Cast<UMaterialExpressionConstant2Vector>(NewExpression))
+		{
+			Const2->R = (*DefaultValueObj)->GetNumberField(TEXT("r"));
+			Const2->G = (*DefaultValueObj)->GetNumberField(TEXT("g"));
+		}
+		else if (UMaterialExpressionConstant3Vector* Const3 = Cast<UMaterialExpressionConstant3Vector>(NewExpression))
+		{
+			Const3->Constant.R = (*DefaultValueObj)->GetNumberField(TEXT("r"));
+			Const3->Constant.G = (*DefaultValueObj)->GetNumberField(TEXT("g"));
+			Const3->Constant.B = (*DefaultValueObj)->GetNumberField(TEXT("b"));
+		}
+		else if (UMaterialExpressionConstant4Vector* Const4 = Cast<UMaterialExpressionConstant4Vector>(NewExpression))
+		{
+			Const4->Constant.R = (*DefaultValueObj)->GetNumberField(TEXT("r"));
+			Const4->Constant.G = (*DefaultValueObj)->GetNumberField(TEXT("g"));
+			Const4->Constant.B = (*DefaultValueObj)->GetNumberField(TEXT("b"));
+			Const4->Constant.A = (*DefaultValueObj)->HasField(TEXT("a")) ? (*DefaultValueObj)->GetNumberField(TEXT("a")) : 1.0f;
+		}
+	}
+	else if (Params->TryGetNumberField(TEXT("default_value"), DefaultValueNum))
+	{
+		// Scalar value
+		if (UMaterialExpressionScalarParameter* ScalarParam = Cast<UMaterialExpressionScalarParameter>(NewExpression))
+		{
+			ScalarParam->DefaultValue = DefaultValueNum;
+		}
+		else if (UMaterialExpressionConstant* Const = Cast<UMaterialExpressionConstant>(NewExpression))
+		{
+			Const->R = DefaultValueNum;
+		}
+	}
+	else if (GetStringParam(Params, TEXT("default_value"), DefaultValueStr, false))
+	{
+		// Texture path or other string value
+		if (UMaterialExpressionTextureSample* TexSample = Cast<UMaterialExpressionTextureSample>(NewExpression))
+		{
+			UTexture2D* Texture = LoadObject<UTexture2D>(nullptr, *DefaultValueStr);
+			if (Texture)
+			{
+				TexSample->Texture = Texture;
+			}
+		}
+	}
+	
+	// Handle TextureSample properties
+	if (UMaterialExpressionTextureSample* TexSample = Cast<UMaterialExpressionTextureSample>(NewExpression))
+	{
+		double TexValue;
+		FString TexStr;
+		bool TexBool;
+		
+		// SamplerType
+		if (GetStringParam(Params, TEXT("sampler_type"), TexStr, false))
+		{
+			if (TexStr.Contains(TEXT("LinearColor"))) TexSample->SamplerType = SAMPLERTYPE_LinearColor;
+			else if (TexStr.Contains(TEXT("LinearGrayscale"))) TexSample->SamplerType = SAMPLERTYPE_LinearGrayscale;
+			else if (TexStr.Contains(TEXT("Grayscale"))) TexSample->SamplerType = SAMPLERTYPE_Grayscale;
+			else if (TexStr.Contains(TEXT("Normal"))) TexSample->SamplerType = SAMPLERTYPE_Normal;
+			else if (TexStr.Contains(TEXT("Alpha"))) TexSample->SamplerType = SAMPLERTYPE_Alpha;
+			else if (TexStr.Contains(TEXT("Masks"))) TexSample->SamplerType = SAMPLERTYPE_Masks;
+			else if (TexStr.Contains(TEXT("Data"))) TexSample->SamplerType = SAMPLERTYPE_Data;
+			else if (TexStr.Contains(TEXT("External"))) TexSample->SamplerType = SAMPLERTYPE_External;
+			else if (TexStr.Contains(TEXT("Color"))) TexSample->SamplerType = SAMPLERTYPE_Color;
+		}
+		else if (Params->TryGetNumberField(TEXT("sampler_type"), TexValue))
+		{
+			TexSample->SamplerType = static_cast<EMaterialSamplerType>(static_cast<uint8>(TexValue));
+		}
+		
+		// SamplerSource
+		if (GetStringParam(Params, TEXT("sampler_source"), TexStr, false))
+		{
+			if (TexStr.Contains(TEXT("FromTextureAsset")) || TexStr.Contains(TEXT("Asset"))) TexSample->SamplerSource = SSM_FromTextureAsset;
+			else if (TexStr.Contains(TEXT("Wrap"))) TexSample->SamplerSource = SSM_Wrap_WorldGroupSettings;
+			else if (TexStr.Contains(TEXT("Clamp"))) TexSample->SamplerSource = SSM_Clamp_WorldGroupSettings;
+		}
+		else if (Params->TryGetNumberField(TEXT("sampler_source"), TexValue))
+		{
+			TexSample->SamplerSource = static_cast<ESamplerSourceMode>(static_cast<uint8>(TexValue));
+		}
+		
+		// MipValueMode
+		if (GetStringParam(Params, TEXT("mip_value_mode"), TexStr, false))
+		{
+			if (TexStr.Contains(TEXT("None"))) TexSample->MipValueMode = TMVM_None;
+			else if (TexStr.Contains(TEXT("MipLevel")) || TexStr.Contains(TEXT("Level"))) TexSample->MipValueMode = TMVM_MipLevel;
+			else if (TexStr.Contains(TEXT("MipBias")) || TexStr.Contains(TEXT("Bias"))) TexSample->MipValueMode = TMVM_MipBias;
+			else if (TexStr.Contains(TEXT("Derivative"))) TexSample->MipValueMode = TMVM_Derivative;
+		}
+		else if (Params->TryGetNumberField(TEXT("mip_value_mode"), TexValue))
+		{
+			TexSample->MipValueMode = static_cast<ETextureMipValueMode>(static_cast<uint8>(TexValue));
+		}
+		
+		// Other numeric/bool properties
+		if (Params->TryGetNumberField(TEXT("const_coordinate"), TexValue))
+			TexSample->ConstCoordinate = static_cast<uint8>(TexValue);
+		if (Params->TryGetNumberField(TEXT("const_mip_value"), TexValue))
+			TexSample->ConstMipValue = static_cast<int32>(TexValue);
+		if (Params->TryGetBoolField(TEXT("automatic_view_mip_bias"), TexBool))
+			TexSample->AutomaticViewMipBias = TexBool;
+	}
+	
+	// Handle Noise properties
+	if (UMaterialExpressionNoise* Noise = Cast<UMaterialExpressionNoise>(NewExpression))
+	{
+		double NoiseValue;
+		bool NoiseBool;
+		if (Params->TryGetNumberField(TEXT("scale"), NoiseValue))
+			Noise->Scale = NoiseValue;
+		if (Params->TryGetNumberField(TEXT("quality"), NoiseValue))
+			Noise->Quality = FMath::Clamp(static_cast<int32>(NoiseValue), 1, 4);
+		if (Params->TryGetNumberField(TEXT("noise_function"), NoiseValue))
+			Noise->NoiseFunction = static_cast<ENoiseFunction>(static_cast<uint8>(NoiseValue));
+		if (Params->TryGetBoolField(TEXT("turbulence"), NoiseBool))
+			Noise->bTurbulence = NoiseBool;
+		if (Params->TryGetNumberField(TEXT("levels"), NoiseValue))
+			Noise->Levels = FMath::Clamp(static_cast<int32>(NoiseValue), 1, 10);
+		if (Params->TryGetNumberField(TEXT("output_min"), NoiseValue))
+			Noise->OutputMin = NoiseValue;
+		if (Params->TryGetNumberField(TEXT("output_max"), NoiseValue))
+			Noise->OutputMax = NoiseValue;
+		if (Params->TryGetNumberField(TEXT("level_scale"), NoiseValue))
+			Noise->LevelScale = FMath::Clamp(static_cast<float>(NoiseValue), 2.0f, 8.0f);
+		if (Params->TryGetBoolField(TEXT("tiling"), NoiseBool))
+			Noise->bTiling = NoiseBool;
+		if (Params->TryGetNumberField(TEXT("repeat_size"), NoiseValue))
+			Noise->RepeatSize = FMath::Max(static_cast<uint32>(NoiseValue), 4u);
+	}
+	
+	// Handle Custom expression properties
+	if (UMaterialExpressionCustom* Custom = Cast<UMaterialExpressionCustom>(NewExpression))
+	{
+		FString CustomStr;
+		double CustomNum;
+		
+		if (GetStringParam(Params, TEXT("code"), CustomStr, false))
+			Custom->Code = CustomStr;
+		if (GetStringParam(Params, TEXT("description"), CustomStr, false))
+			Custom->Description = CustomStr;
+		
+		// OutputType
+		if (GetStringParam(Params, TEXT("output_type"), CustomStr, false))
+		{
+			if (CustomStr.Contains(TEXT("Float1")) || CustomStr.Contains(TEXT("Scalar"))) Custom->OutputType = CMOT_Float1;
+			else if (CustomStr.Contains(TEXT("Float2"))) Custom->OutputType = CMOT_Float2;
+			else if (CustomStr.Contains(TEXT("Float3"))) Custom->OutputType = CMOT_Float3;
+			else if (CustomStr.Contains(TEXT("Float4"))) Custom->OutputType = CMOT_Float4;
+			else if (CustomStr.Contains(TEXT("MaterialAttributes"))) Custom->OutputType = CMOT_MaterialAttributes;
+		}
+		else if (Params->TryGetNumberField(TEXT("output_type"), CustomNum))
+		{
+			Custom->OutputType = static_cast<ECustomMaterialOutputType>(static_cast<uint8>(CustomNum));
+		}
+		
+		// Inputs array
+		const TArray<TSharedPtr<FJsonValue>>* InputsArray;
+		if (Params->TryGetArrayField(TEXT("inputs"), InputsArray))
+		{
+			Custom->Inputs.Empty();
+			UE_LOG(LogTemp, Warning, TEXT("ECABridge AddMaterialNode: Setting Custom inputs, array size: %d"), InputsArray->Num());
+			for (const TSharedPtr<FJsonValue>& InputValue : *InputsArray)
+			{
+				const TSharedPtr<FJsonObject>* InputObj;
+				if (InputValue->TryGetObject(InputObj))
+				{
+					FCustomInput NewInput;
+					FString InputName;
+					if ((*InputObj)->TryGetStringField(TEXT("input_name"), InputName))
+					{
+						NewInput.InputName = FName(*InputName);
+						UE_LOG(LogTemp, Warning, TEXT("ECABridge AddMaterialNode: Set input name from object: '%s'"), *InputName);
+					}
+					else
+					{
+						UE_LOG(LogTemp, Warning, TEXT("ECABridge AddMaterialNode: input_name field not found in object"));
+					}
+					Custom->Inputs.Add(NewInput);
+				}
+				else
+				{
+					FString InputName = InputValue->AsString();
+					if (!InputName.IsEmpty())
+					{
+						FCustomInput NewInput;
+						NewInput.InputName = FName(*InputName);
+						UE_LOG(LogTemp, Warning, TEXT("ECABridge AddMaterialNode: Set input name from string: '%s'"), *InputName);
+						Custom->Inputs.Add(NewInput);
+					}
+				}
+			}
+			UE_LOG(LogTemp, Warning, TEXT("ECABridge AddMaterialNode: Final inputs count: %d"), Custom->Inputs.Num());
+			for (int32 i = 0; i < Custom->Inputs.Num(); i++)
+			{
+				UE_LOG(LogTemp, Warning, TEXT("ECABridge AddMaterialNode: Input[%d].InputName = '%s'"), i, *Custom->Inputs[i].InputName.ToString());
+			}
+		}
+		
+		// Additional outputs
+		const TArray<TSharedPtr<FJsonValue>>* OutputsArray;
+		if (Params->TryGetArrayField(TEXT("additional_outputs"), OutputsArray))
+		{
+			Custom->AdditionalOutputs.Empty();
+			for (const TSharedPtr<FJsonValue>& OutputValue : *OutputsArray)
+			{
+				const TSharedPtr<FJsonObject>* OutputObj;
+				if (OutputValue->TryGetObject(OutputObj))
+				{
+					FCustomOutput NewOutput;
+					FString OutputName;
+					if ((*OutputObj)->TryGetStringField(TEXT("output_name"), OutputName))
+					{
+						NewOutput.OutputName = FName(*OutputName);
+					}
+					double OutputType;
+					if ((*OutputObj)->TryGetNumberField(TEXT("output_type"), OutputType))
+					{
+						NewOutput.OutputType = static_cast<ECustomMaterialOutputType>(static_cast<uint8>(OutputType));
+					}
+					Custom->AdditionalOutputs.Add(NewOutput);
+				}
+			}
+		}
+		
+		// Additional defines
+		const TArray<TSharedPtr<FJsonValue>>* DefinesArray;
+		if (Params->TryGetArrayField(TEXT("additional_defines"), DefinesArray))
+		{
+			Custom->AdditionalDefines.Empty();
+			for (const TSharedPtr<FJsonValue>& DefineValue : *DefinesArray)
+			{
+				const TSharedPtr<FJsonObject>* DefineObj;
+				if (DefineValue->TryGetObject(DefineObj))
+				{
+					FCustomDefine NewDefine;
+					(*DefineObj)->TryGetStringField(TEXT("define_name"), NewDefine.DefineName);
+					(*DefineObj)->TryGetStringField(TEXT("define_value"), NewDefine.DefineValue);
+					Custom->AdditionalDefines.Add(NewDefine);
+				}
+			}
+		}
+		
+		// Include file paths
+		const TArray<TSharedPtr<FJsonValue>>* IncludesArray;
+		if (Params->TryGetArrayField(TEXT("include_file_paths"), IncludesArray))
+		{
+			Custom->IncludeFilePaths.Empty();
+			for (const TSharedPtr<FJsonValue>& IncludeValue : *IncludesArray)
+			{
+				FString IncludePath = IncludeValue->AsString();
+				if (!IncludePath.IsEmpty())
+				{
+					Custom->IncludeFilePaths.Add(IncludePath);
+				}
+			}
+		}
+		
+		// Rebuild outputs and reconstruct graph node to ensure input/output pins are properly created
+		Custom->RebuildOutputs();
+		if (Custom->GraphNode)
+		{
+			Custom->GraphNode->ReconstructNode();
+		}
+	}
+	
+	// Add to material
+	Material->GetEditorOnlyData()->ExpressionCollection.AddExpression(NewExpression);
+	
+	ClearMaterialCompilationErrors(Material);
+	Material->PostEditChange();
+	Material->MarkPackageDirty();
+	
+	// Refresh the material editor UI if open
+	RefreshMaterialEditorUI(Material);
+	
+	TSharedPtr<FJsonObject> Result = MakeResult();
+	Result->SetStringField(TEXT("node_id"), NewExpression->MaterialExpressionGuid.ToString());
+	Result->SetStringField(TEXT("node_type"), NodeType);
+	Result->SetStringField(TEXT("node_class"), NewExpression->GetClass()->GetName());
+	
+	// Include any compilation errors
+	AddCompilationErrorsToResult(Result, Material);
+	
+	return FECACommandResult::Success(Result);
+}
+
+//------------------------------------------------------------------------------
+// ConnectMaterialNodes
+//------------------------------------------------------------------------------
+
+FECACommandResult FECACommand_ConnectMaterialNodes::Execute(const TSharedPtr<FJsonObject>& Params)
+{
+	FString MaterialPath;
+	if (!GetStringParam(Params, TEXT("material_path"), MaterialPath))
+	{
+		return FECACommandResult::Error(TEXT("Missing required parameter: material_path"));
+	}
+	
+	FString SourceNodeId;
+	if (!GetStringParam(Params, TEXT("source_node_id"), SourceNodeId))
+	{
+		return FECACommandResult::Error(TEXT("Missing required parameter: source_node_id"));
+	}
+	
+	FString TargetNodeId;
+	if (!GetStringParam(Params, TEXT("target_node_id"), TargetNodeId))
+	{
+		return FECACommandResult::Error(TEXT("Missing required parameter: target_node_id"));
+	}
+	
+	FString TargetInput;
+	if (!GetStringParam(Params, TEXT("target_input"), TargetInput))
+	{
+		return FECACommandResult::Error(TEXT("Missing required parameter: target_input"));
+	}
+	
+	FString SourceOutput = TEXT("0");
+	GetStringParam(Params, TEXT("source_output"), SourceOutput, false);
+	
+	UMaterial* Material = LoadMaterialByPath(MaterialPath);
+	if (!Material)
+	{
+		return FECACommandResult::Error(FString::Printf(TEXT("Material not found: %s"), *MaterialPath));
+	}
+	
+	// Find source expression
+	FGuid SourceGuid;
+	if (!FGuid::Parse(SourceNodeId, SourceGuid))
+	{
+		return FECACommandResult::Error(FString::Printf(TEXT("Invalid source node GUID: %s"), *SourceNodeId));
+	}
+	
+	UMaterialExpression* SourceExpr = FindExpressionByGuid(Material, SourceGuid);
+	if (!SourceExpr)
+	{
+		return FECACommandResult::Error(FString::Printf(TEXT("Source node not found: %s"), *SourceNodeId));
+	}
+	
+	int32 OutputIndex = ParseOutputIndex(SourceOutput, SourceExpr);
+	
+	Material->PreEditChange(nullptr);
+	
+	// Connect to material input or another expression
+	if (TargetNodeId.Equals(TEXT("material"), ESearchCase::IgnoreCase))
+	{
+		// Connect to material input
+		FExpressionInput* MaterialInput = GetMaterialInput(Material, TargetInput);
+		if (!MaterialInput)
+		{
+			ClearMaterialCompilationErrors(Material);
+			Material->PostEditChange();
+			return FECACommandResult::Error(FString::Printf(TEXT("Unknown material input: %s"), *TargetInput));
+		}
+		
+		MaterialInput->Expression = SourceExpr;
+		MaterialInput->OutputIndex = OutputIndex;
+	}
+	else
+	{
+		// Connect to another expression
+		FGuid TargetGuid;
+		if (!FGuid::Parse(TargetNodeId, TargetGuid))
+		{
+			ClearMaterialCompilationErrors(Material);
+			Material->PostEditChange();
+			return FECACommandResult::Error(FString::Printf(TEXT("Invalid target node GUID: %s"), *TargetNodeId));
+		}
+		
+		UMaterialExpression* TargetExpr = FindExpressionByGuid(Material, TargetGuid);
+		if (!TargetExpr)
+		{
+			ClearMaterialCompilationErrors(Material);
+			Material->PostEditChange();
+			return FECACommandResult::Error(FString::Printf(TEXT("Target node not found: %s"), *TargetNodeId));
+		}
+		
+		// Find the target input
+		FExpressionInput* TargetInputPtr = nullptr;
+		
+		// Try by index first
+		if (TargetInput.IsNumeric())
+		{
+			int32 InputIndex = FCString::Atoi(*TargetInput);
+			TargetInputPtr = TargetExpr->GetInput(InputIndex);
+		}
+		else
+		{
+			// Try by name - iterate through all inputs
+			for (int32 i = 0; ; i++)
+			{
+				FExpressionInput* Input = TargetExpr->GetInput(i);
+				if (!Input)
+				{
+					break;
+				}
+				FName InputName = TargetExpr->GetInputName(i);
+				if (InputName.ToString().Equals(TargetInput, ESearchCase::IgnoreCase))
+				{
+					TargetInputPtr = Input;
+					break;
+				}
+			}
+			
+			// Try common input names for specific expression types
+			if (!TargetInputPtr)
+			{
+				if (UMaterialExpressionMultiply* Mult = Cast<UMaterialExpressionMultiply>(TargetExpr))
+				{
+					if (TargetInput.Equals(TEXT("A"), ESearchCase::IgnoreCase)) TargetInputPtr = &Mult->A;
+					else if (TargetInput.Equals(TEXT("B"), ESearchCase::IgnoreCase)) TargetInputPtr = &Mult->B;
+				}
+				else if (UMaterialExpressionAdd* Add = Cast<UMaterialExpressionAdd>(TargetExpr))
+				{
+					if (TargetInput.Equals(TEXT("A"), ESearchCase::IgnoreCase)) TargetInputPtr = &Add->A;
+					else if (TargetInput.Equals(TEXT("B"), ESearchCase::IgnoreCase)) TargetInputPtr = &Add->B;
+				}
+				else if (UMaterialExpressionSubtract* Sub = Cast<UMaterialExpressionSubtract>(TargetExpr))
+				{
+					if (TargetInput.Equals(TEXT("A"), ESearchCase::IgnoreCase)) TargetInputPtr = &Sub->A;
+					else if (TargetInput.Equals(TEXT("B"), ESearchCase::IgnoreCase)) TargetInputPtr = &Sub->B;
+				}
+				else if (UMaterialExpressionDivide* Div = Cast<UMaterialExpressionDivide>(TargetExpr))
+				{
+					if (TargetInput.Equals(TEXT("A"), ESearchCase::IgnoreCase)) TargetInputPtr = &Div->A;
+					else if (TargetInput.Equals(TEXT("B"), ESearchCase::IgnoreCase)) TargetInputPtr = &Div->B;
+				}
+				else if (UMaterialExpressionPower* Pow = Cast<UMaterialExpressionPower>(TargetExpr))
+				{
+					if (TargetInput.Equals(TEXT("Base"), ESearchCase::IgnoreCase)) TargetInputPtr = &Pow->Base;
+					else if (TargetInput.Equals(TEXT("Exponent"), ESearchCase::IgnoreCase) || TargetInput.Equals(TEXT("Exp"), ESearchCase::IgnoreCase)) TargetInputPtr = &Pow->Exponent;
+				}
+				else if (UMaterialExpressionClamp* Clamp = Cast<UMaterialExpressionClamp>(TargetExpr))
+				{
+					if (TargetInput.Equals(TEXT("Input"), ESearchCase::IgnoreCase) || TargetInput.Equals(TEXT("0"), ESearchCase::IgnoreCase)) TargetInputPtr = &Clamp->Input;
+					else if (TargetInput.Equals(TEXT("Min"), ESearchCase::IgnoreCase)) TargetInputPtr = &Clamp->Min;
+					else if (TargetInput.Equals(TEXT("Max"), ESearchCase::IgnoreCase)) TargetInputPtr = &Clamp->Max;
+				}
+				else if (UMaterialExpressionIf* If = Cast<UMaterialExpressionIf>(TargetExpr))
+				{
+					if (TargetInput.Equals(TEXT("A"), ESearchCase::IgnoreCase)) TargetInputPtr = &If->A;
+					else if (TargetInput.Equals(TEXT("B"), ESearchCase::IgnoreCase)) TargetInputPtr = &If->B;
+					else if (TargetInput.Equals(TEXT("AGreaterThanB"), ESearchCase::IgnoreCase) || TargetInput.Equals(TEXT("A>B"), ESearchCase::IgnoreCase)) TargetInputPtr = &If->AGreaterThanB;
+					else if (TargetInput.Equals(TEXT("AEqualsB"), ESearchCase::IgnoreCase) || TargetInput.Equals(TEXT("A==B"), ESearchCase::IgnoreCase) || TargetInput.Equals(TEXT("A=B"), ESearchCase::IgnoreCase)) TargetInputPtr = &If->AEqualsB;
+					else if (TargetInput.Equals(TEXT("ALessThanB"), ESearchCase::IgnoreCase) || TargetInput.Equals(TEXT("A<B"), ESearchCase::IgnoreCase)) TargetInputPtr = &If->ALessThanB;
+				}
+				else if (UMaterialExpressionLinearInterpolate* Lerp = Cast<UMaterialExpressionLinearInterpolate>(TargetExpr))
+				{
+					if (TargetInput.Equals(TEXT("A"), ESearchCase::IgnoreCase)) TargetInputPtr = &Lerp->A;
+					else if (TargetInput.Equals(TEXT("B"), ESearchCase::IgnoreCase)) TargetInputPtr = &Lerp->B;
+					else if (TargetInput.Equals(TEXT("Alpha"), ESearchCase::IgnoreCase)) TargetInputPtr = &Lerp->Alpha;
+				}
+			}
+		}
+		
+		if (!TargetInputPtr)
+		{
+			ClearMaterialCompilationErrors(Material);
+			Material->PostEditChange();
+			return FECACommandResult::Error(FString::Printf(TEXT("Input not found on target node: %s"), *TargetInput));
+		}
+		
+		TargetInputPtr->Expression = SourceExpr;
+		TargetInputPtr->OutputIndex = OutputIndex;
+	}
+	
+	ClearMaterialCompilationErrors(Material);
+	Material->PostEditChange();
+	Material->MarkPackageDirty();
+	
+	// Refresh the material editor UI if open
+	RefreshMaterialEditorUI(Material);
+	
+	TSharedPtr<FJsonObject> Result = MakeResult();
+	Result->SetStringField(TEXT("source_node_id"), SourceNodeId);
+	Result->SetStringField(TEXT("target_node_id"), TargetNodeId);
+	Result->SetStringField(TEXT("target_input"), TargetInput);
+	Result->SetNumberField(TEXT("source_output_index"), OutputIndex);
+	
+	// Include any compilation errors
+	AddCompilationErrorsToResult(Result, Material);
+	
+	return FECACommandResult::Success(Result);
+}
+
+//------------------------------------------------------------------------------
+// GetMaterialNodes
+//------------------------------------------------------------------------------
+
+FECACommandResult FECACommand_GetMaterialNodes::Execute(const TSharedPtr<FJsonObject>& Params)
+{
+	FString MaterialPath;
+	if (!GetStringParam(Params, TEXT("material_path"), MaterialPath))
+	{
+		return FECACommandResult::Error(TEXT("Missing required parameter: material_path"));
+	}
+	
+	FString NodeTypeFilter;
+	GetStringParam(Params, TEXT("node_type_filter"), NodeTypeFilter, false);
+	
+	UMaterial* Material = LoadMaterialByPath(MaterialPath);
+	if (!Material)
+	{
+		return FECACommandResult::Error(FString::Printf(TEXT("Material not found: %s"), *MaterialPath));
+	}
+	
+	TArray<TSharedPtr<FJsonValue>> NodesArray;
+	
+	for (UMaterialExpression* Expr : Material->GetEditorOnlyData()->ExpressionCollection.Expressions)
+	{
+		if (!Expr) continue;
+		
+		// Apply filter if specified
+		if (!NodeTypeFilter.IsEmpty())
+		{
+			FString ExprClassName = Expr->GetClass()->GetName();
+			if (!ExprClassName.Contains(NodeTypeFilter, ESearchCase::IgnoreCase))
+			{
+				continue;
+			}
+		}
+		
+		NodesArray.Add(MakeShared<FJsonValueObject>(ExpressionToJson(Expr)));
+	}
+	
+	// Also include material input connections
+	TSharedPtr<FJsonObject> MaterialInputs = MakeShared<FJsonObject>();
+	auto AddMaterialInput = [&](const TCHAR* Name, const FExpressionInput& Input)
+	{
+		if (Input.Expression)
+		{
+			TSharedPtr<FJsonObject> InputInfo = MakeShared<FJsonObject>();
+			InputInfo->SetStringField(TEXT("connected_node_id"), Input.Expression->MaterialExpressionGuid.ToString());
+			InputInfo->SetNumberField(TEXT("connected_output_index"), Input.OutputIndex);
+			MaterialInputs->SetObjectField(Name, InputInfo);
+		}
+	};
+	
+	AddMaterialInput(TEXT("BaseColor"), Material->GetEditorOnlyData()->BaseColor);
+	AddMaterialInput(TEXT("Metallic"), Material->GetEditorOnlyData()->Metallic);
+	AddMaterialInput(TEXT("Specular"), Material->GetEditorOnlyData()->Specular);
+	AddMaterialInput(TEXT("Roughness"), Material->GetEditorOnlyData()->Roughness);
+	AddMaterialInput(TEXT("Normal"), Material->GetEditorOnlyData()->Normal);
+	AddMaterialInput(TEXT("EmissiveColor"), Material->GetEditorOnlyData()->EmissiveColor);
+	AddMaterialInput(TEXT("Opacity"), Material->GetEditorOnlyData()->Opacity);
+	AddMaterialInput(TEXT("OpacityMask"), Material->GetEditorOnlyData()->OpacityMask);
+	AddMaterialInput(TEXT("WorldPositionOffset"), Material->GetEditorOnlyData()->WorldPositionOffset);
+	AddMaterialInput(TEXT("AmbientOcclusion"), Material->GetEditorOnlyData()->AmbientOcclusion);
+	AddMaterialInput(TEXT("SubsurfaceColor"), Material->GetEditorOnlyData()->SubsurfaceColor);
+	AddMaterialInput(TEXT("Refraction"), Material->GetEditorOnlyData()->Refraction);
+	AddMaterialInput(TEXT("PixelDepthOffset"), Material->GetEditorOnlyData()->PixelDepthOffset);
+	AddMaterialInput(TEXT("Displacement"), Material->GetEditorOnlyData()->Displacement);
+	AddMaterialInput(TEXT("ShadingModelFromMaterialExpression"), Material->GetEditorOnlyData()->ShadingModelFromMaterialExpression);
+	AddMaterialInput(TEXT("FrontMaterial"), Material->GetEditorOnlyData()->FrontMaterial);
+	AddMaterialInput(TEXT("SurfaceThickness"), Material->GetEditorOnlyData()->SurfaceThickness);
+	AddMaterialInput(TEXT("Tangent"), Material->GetEditorOnlyData()->Tangent);
+	AddMaterialInput(TEXT("Anisotropy"), Material->GetEditorOnlyData()->Anisotropy);
+	AddMaterialInput(TEXT("ClearCoat"), Material->GetEditorOnlyData()->ClearCoat);
+	AddMaterialInput(TEXT("ClearCoatRoughness"), Material->GetEditorOnlyData()->ClearCoatRoughness);
+	
+	TSharedPtr<FJsonObject> Result = MakeResult();
+	Result->SetArrayField(TEXT("nodes"), NodesArray);
+	Result->SetNumberField(TEXT("node_count"), NodesArray.Num());
+	Result->SetObjectField(TEXT("material_inputs"), MaterialInputs);
+	
+	// Include any compilation errors
+	AddCompilationErrorsToResult(Result, Material);
+	
+	return FECACommandResult::Success(Result);
+}
+
+//------------------------------------------------------------------------------
+// GetMaterialNodeInfo
+//------------------------------------------------------------------------------
+
+FECACommandResult FECACommand_GetMaterialNodeInfo::Execute(const TSharedPtr<FJsonObject>& Params)
+{
+	FString MaterialPath;
+	if (!GetStringParam(Params, TEXT("material_path"), MaterialPath))
+	{
+		return FECACommandResult::Error(TEXT("Missing required parameter: material_path"));
+	}
+	
+	FString NodeId;
+	if (!GetStringParam(Params, TEXT("node_id"), NodeId))
+	{
+		return FECACommandResult::Error(TEXT("Missing required parameter: node_id"));
+	}
+	
+	UMaterial* Material = LoadMaterialByPath(MaterialPath);
+	if (!Material)
+	{
+		return FECACommandResult::Error(FString::Printf(TEXT("Material not found: %s"), *MaterialPath));
+	}
+	
+	FGuid NodeGuid;
+	if (!FGuid::Parse(NodeId, NodeGuid))
+	{
+		return FECACommandResult::Error(FString::Printf(TEXT("Invalid node GUID: %s"), *NodeId));
+	}
+	
+	UMaterialExpression* Expression = FindExpressionByGuid(Material, NodeGuid);
+	if (!Expression)
+	{
+		return FECACommandResult::Error(FString::Printf(TEXT("Node not found: %s"), *NodeId));
+	}
+	
+	TSharedPtr<FJsonObject> Result = ExpressionToJson(Expression);
+	return FECACommandResult::Success(Result);
+}
+
+//------------------------------------------------------------------------------
+// DeleteMaterialNode
+//------------------------------------------------------------------------------
+
+FECACommandResult FECACommand_DeleteMaterialNode::Execute(const TSharedPtr<FJsonObject>& Params)
+{
+	FString MaterialPath;
+	if (!GetStringParam(Params, TEXT("material_path"), MaterialPath))
+	{
+		return FECACommandResult::Error(TEXT("Missing required parameter: material_path"));
+	}
+	
+	FString NodeId;
+	if (!GetStringParam(Params, TEXT("node_id"), NodeId))
+	{
+		return FECACommandResult::Error(TEXT("Missing required parameter: node_id"));
+	}
+	
+	UMaterial* Material = LoadMaterialByPath(MaterialPath);
+	if (!Material)
+	{
+		return FECACommandResult::Error(FString::Printf(TEXT("Material not found: %s"), *MaterialPath));
+	}
+	
+	FGuid NodeGuid;
+	if (!FGuid::Parse(NodeId, NodeGuid))
+	{
+		return FECACommandResult::Error(FString::Printf(TEXT("Invalid node GUID: %s"), *NodeId));
+	}
+	
+	UMaterialExpression* Expression = FindExpressionByGuid(Material, NodeGuid);
+	if (!Expression)
+	{
+		return FECACommandResult::Error(FString::Printf(TEXT("Node not found: %s"), *NodeId));
+	}
+	
+	Material->PreEditChange(nullptr);
+	
+	// Remove all connections to this expression from material inputs
+	auto ClearMaterialInput = [&](FExpressionInput& Input)
+	{
+		if (Input.Expression == Expression)
+		{
+			Input.Expression = nullptr;
+			Input.OutputIndex = 0;
+		}
+	};
+	
+	ClearMaterialInput(Material->GetEditorOnlyData()->BaseColor);
+	ClearMaterialInput(Material->GetEditorOnlyData()->Metallic);
+	ClearMaterialInput(Material->GetEditorOnlyData()->Specular);
+	ClearMaterialInput(Material->GetEditorOnlyData()->Roughness);
+	ClearMaterialInput(Material->GetEditorOnlyData()->Normal);
+	ClearMaterialInput(Material->GetEditorOnlyData()->EmissiveColor);
+	ClearMaterialInput(Material->GetEditorOnlyData()->Opacity);
+	ClearMaterialInput(Material->GetEditorOnlyData()->OpacityMask);
+	ClearMaterialInput(Material->GetEditorOnlyData()->WorldPositionOffset);
+	ClearMaterialInput(Material->GetEditorOnlyData()->AmbientOcclusion);
+	ClearMaterialInput(Material->GetEditorOnlyData()->SubsurfaceColor);
+	ClearMaterialInput(Material->GetEditorOnlyData()->Refraction);
+	ClearMaterialInput(Material->GetEditorOnlyData()->PixelDepthOffset);
+	ClearMaterialInput(Material->GetEditorOnlyData()->Displacement);
+	ClearMaterialInput(Material->GetEditorOnlyData()->ShadingModelFromMaterialExpression);
+	ClearMaterialInput(Material->GetEditorOnlyData()->FrontMaterial);
+	ClearMaterialInput(Material->GetEditorOnlyData()->SurfaceThickness);
+	ClearMaterialInput(Material->GetEditorOnlyData()->Tangent);
+	ClearMaterialInput(Material->GetEditorOnlyData()->Anisotropy);
+	ClearMaterialInput(Material->GetEditorOnlyData()->ClearCoat);
+	ClearMaterialInput(Material->GetEditorOnlyData()->ClearCoatRoughness);
+	
+	// Remove connections from other expressions
+	for (UMaterialExpression* OtherExpr : Material->GetEditorOnlyData()->ExpressionCollection.Expressions)
+	{
+		if (OtherExpr && OtherExpr != Expression)
+		{
+			for (int32 i = 0; ; i++)
+			{
+				FExpressionInput* Input = OtherExpr->GetInput(i);
+				if (!Input)
+				{
+					break;
+				}
+				if (Input->Expression == Expression)
+				{
+					Input->Expression = nullptr;
+					Input->OutputIndex = 0;
+				}
+			}
+		}
+	}
+	
+	// Remove from material
+	Material->GetEditorOnlyData()->ExpressionCollection.RemoveExpression(Expression);
+	
+	ClearMaterialCompilationErrors(Material);
+	Material->PostEditChange();
+	Material->MarkPackageDirty();
+	
+	// Refresh the material editor UI if open
+	RefreshMaterialEditorUI(Material);
+	
+	TSharedPtr<FJsonObject> Result = MakeResult();
+	Result->SetStringField(TEXT("deleted_node_id"), NodeId);
+	
+	return FECACommandResult::Success(Result);
+}
+
+//------------------------------------------------------------------------------
+// DisconnectMaterialNode
+//------------------------------------------------------------------------------
+
+FECACommandResult FECACommand_DisconnectMaterialNode::Execute(const TSharedPtr<FJsonObject>& Params)
+{
+	FString MaterialPath;
+	if (!GetStringParam(Params, TEXT("material_path"), MaterialPath))
+	{
+		return FECACommandResult::Error(TEXT("Missing required parameter: material_path"));
+	}
+	
+	FString NodeId;
+	if (!GetStringParam(Params, TEXT("node_id"), NodeId))
+	{
+		return FECACommandResult::Error(TEXT("Missing required parameter: node_id"));
+	}
+	
+	FString InputName;
+	GetStringParam(Params, TEXT("input_name"), InputName, false);
+	
+	UMaterial* Material = LoadMaterialByPath(MaterialPath);
+	if (!Material)
+	{
+		return FECACommandResult::Error(FString::Printf(TEXT("Material not found: %s"), *MaterialPath));
+	}
+	
+	FGuid NodeGuid;
+	if (!FGuid::Parse(NodeId, NodeGuid))
+	{
+		return FECACommandResult::Error(FString::Printf(TEXT("Invalid node GUID: %s"), *NodeId));
+	}
+	
+	UMaterialExpression* Expression = FindExpressionByGuid(Material, NodeGuid);
+	if (!Expression)
+	{
+		return FECACommandResult::Error(FString::Printf(TEXT("Node not found: %s"), *NodeId));
+	}
+	
+	Material->PreEditChange(nullptr);
+	
+	int32 DisconnectedCount = 0;
+	
+	for (int32 i = 0; ; i++)
+	{
+		FExpressionInput* Input = Expression->GetInput(i);
+		if (!Input)
+		{
+			break;
+		}
+		if (!Input->Expression) continue;
+		
+		// If specific input name given, only disconnect that one
+		if (!InputName.IsEmpty())
+		{
+			FName ThisInputName = Expression->GetInputName(i);
+			FString InputNameStr = ThisInputName.IsNone() ? FString::FromInt(i) : ThisInputName.ToString();
+			if (!InputNameStr.Equals(InputName, ESearchCase::IgnoreCase))
+			{
+				continue;
+			}
+		}
+		
+		Input->Expression = nullptr;
+		Input->OutputIndex = 0;
+		DisconnectedCount++;
+	}
+	
+	ClearMaterialCompilationErrors(Material);
+	Material->PostEditChange();
+	Material->MarkPackageDirty();
+	
+	// Refresh the material editor UI if open
+	RefreshMaterialEditorUI(Material);
+	
+	TSharedPtr<FJsonObject> Result = MakeResult();
+	Result->SetStringField(TEXT("node_id"), NodeId);
+	Result->SetNumberField(TEXT("disconnected_count"), DisconnectedCount);
+	
+	return FECACommandResult::Success(Result);
+}
+
+//------------------------------------------------------------------------------
+// SetMaterialNodeProperty
+//------------------------------------------------------------------------------
+
+FECACommandResult FECACommand_SetMaterialNodeProperty::Execute(const TSharedPtr<FJsonObject>& Params)
+{
+	FString MaterialPath;
+	if (!GetStringParam(Params, TEXT("material_path"), MaterialPath))
+	{
+		return FECACommandResult::Error(TEXT("Missing required parameter: material_path"));
+	}
+	
+	FString NodeId;
+	if (!GetStringParam(Params, TEXT("node_id"), NodeId))
+	{
+		return FECACommandResult::Error(TEXT("Missing required parameter: node_id"));
+	}
+	
+	FString PropertyName;
+	if (!GetStringParam(Params, TEXT("property_name"), PropertyName))
+	{
+		return FECACommandResult::Error(TEXT("Missing required parameter: property_name"));
+	}
+	
+	UMaterial* Material = LoadMaterialByPath(MaterialPath);
+	if (!Material)
+	{
+		return FECACommandResult::Error(FString::Printf(TEXT("Material not found: %s"), *MaterialPath));
+	}
+	
+	FGuid NodeGuid;
+	if (!FGuid::Parse(NodeId, NodeGuid))
+	{
+		return FECACommandResult::Error(FString::Printf(TEXT("Invalid node GUID: %s"), *NodeId));
+	}
+	
+	UMaterialExpression* Expression = FindExpressionByGuid(Material, NodeGuid);
+	if (!Expression)
+	{
+		return FECACommandResult::Error(FString::Printf(TEXT("Node not found: %s"), *NodeId));
+	}
+	
+	Material->PreEditChange(nullptr);
+	
+	bool bPropertySet = false;
+	
+	// Handle common properties
+	const TSharedPtr<FJsonObject>* ValueObj;
+	double ValueNum;
+	FString ValueStr;
+	bool ValueBool;
+	
+	if (PropertyName.Equals(TEXT("ParameterName"), ESearchCase::IgnoreCase))
+	{
+		if (GetStringParam(Params, TEXT("value"), ValueStr, false))
+		{
+			if (UMaterialExpressionScalarParameter* ScalarParam = Cast<UMaterialExpressionScalarParameter>(Expression))
+			{
+				ScalarParam->ParameterName = FName(*ValueStr);
+				bPropertySet = true;
+			}
+			else if (UMaterialExpressionVectorParameter* VectorParam = Cast<UMaterialExpressionVectorParameter>(Expression))
+			{
+				VectorParam->ParameterName = FName(*ValueStr);
+				bPropertySet = true;
+			}
+			else if (UMaterialExpressionTextureSampleParameter2D* TexParam = Cast<UMaterialExpressionTextureSampleParameter2D>(Expression))
+			{
+				TexParam->ParameterName = FName(*ValueStr);
+				bPropertySet = true;
+			}
+		}
+	}
+	else if (PropertyName.Equals(TEXT("DefaultValue"), ESearchCase::IgnoreCase) || PropertyName.Equals(TEXT("Value"), ESearchCase::IgnoreCase))
+	{
+		if (Params->TryGetObjectField(TEXT("value"), ValueObj))
+		{
+			if (UMaterialExpressionVectorParameter* VectorParam = Cast<UMaterialExpressionVectorParameter>(Expression))
+			{
+				VectorParam->DefaultValue.R = (*ValueObj)->GetNumberField(TEXT("r"));
+				VectorParam->DefaultValue.G = (*ValueObj)->GetNumberField(TEXT("g"));
+				VectorParam->DefaultValue.B = (*ValueObj)->GetNumberField(TEXT("b"));
+				if ((*ValueObj)->HasField(TEXT("a")))
+					VectorParam->DefaultValue.A = (*ValueObj)->GetNumberField(TEXT("a"));
+				bPropertySet = true;
+			}
+			else if (UMaterialExpressionConstant2Vector* Const2 = Cast<UMaterialExpressionConstant2Vector>(Expression))
+			{
+				Const2->R = (*ValueObj)->GetNumberField(TEXT("r"));
+				Const2->G = (*ValueObj)->GetNumberField(TEXT("g"));
+				bPropertySet = true;
+			}
+			else if (UMaterialExpressionConstant3Vector* Const3 = Cast<UMaterialExpressionConstant3Vector>(Expression))
+			{
+				Const3->Constant.R = (*ValueObj)->GetNumberField(TEXT("r"));
+				Const3->Constant.G = (*ValueObj)->GetNumberField(TEXT("g"));
+				Const3->Constant.B = (*ValueObj)->GetNumberField(TEXT("b"));
+				bPropertySet = true;
+			}
+			else if (UMaterialExpressionConstant4Vector* Const4 = Cast<UMaterialExpressionConstant4Vector>(Expression))
+			{
+				Const4->Constant.R = (*ValueObj)->GetNumberField(TEXT("r"));
+				Const4->Constant.G = (*ValueObj)->GetNumberField(TEXT("g"));
+				Const4->Constant.B = (*ValueObj)->GetNumberField(TEXT("b"));
+				Const4->Constant.A = (*ValueObj)->HasField(TEXT("a")) ? (*ValueObj)->GetNumberField(TEXT("a")) : 1.0f;
+				bPropertySet = true;
+			}
+		}
+		else if (Params->TryGetNumberField(TEXT("value"), ValueNum))
+		{
+			if (UMaterialExpressionScalarParameter* ScalarParam = Cast<UMaterialExpressionScalarParameter>(Expression))
+			{
+				ScalarParam->DefaultValue = ValueNum;
+				bPropertySet = true;
+			}
+			else if (UMaterialExpressionConstant* Const = Cast<UMaterialExpressionConstant>(Expression))
+			{
+				Const->R = ValueNum;
+				bPropertySet = true;
+			}
+		}
+	}
+	else if (PropertyName.Equals(TEXT("Texture"), ESearchCase::IgnoreCase))
+	{
+		if (GetStringParam(Params, TEXT("value"), ValueStr, false))
+		{
+			UTexture2D* Texture = LoadObject<UTexture2D>(nullptr, *ValueStr);
+			if (Texture)
+			{
+				if (UMaterialExpressionTextureSample* TexSample = Cast<UMaterialExpressionTextureSample>(Expression))
+				{
+					TexSample->Texture = Texture;
+					bPropertySet = true;
+				}
+			}
+		}
+	}
+	else if (PropertyName.Equals(TEXT("SamplerType"), ESearchCase::IgnoreCase))
+	{
+		if (UMaterialExpressionTextureBase* TexBase = Cast<UMaterialExpressionTextureBase>(Expression))
+		{
+			if (GetStringParam(Params, TEXT("value"), ValueStr, false))
+			{
+				// Map string names to enum values
+				if (ValueStr.Contains(TEXT("LinearColor"))) TexBase->SamplerType = SAMPLERTYPE_LinearColor;
+				else if (ValueStr.Contains(TEXT("LinearGrayscale"))) TexBase->SamplerType = SAMPLERTYPE_LinearGrayscale;
+				else if (ValueStr.Contains(TEXT("Grayscale"))) TexBase->SamplerType = SAMPLERTYPE_Grayscale;
+				else if (ValueStr.Contains(TEXT("Normal"))) TexBase->SamplerType = SAMPLERTYPE_Normal;
+				else if (ValueStr.Contains(TEXT("Alpha"))) TexBase->SamplerType = SAMPLERTYPE_Alpha;
+				else if (ValueStr.Contains(TEXT("Masks"))) TexBase->SamplerType = SAMPLERTYPE_Masks;
+				else if (ValueStr.Contains(TEXT("Data"))) TexBase->SamplerType = SAMPLERTYPE_Data;
+				else if (ValueStr.Contains(TEXT("External"))) TexBase->SamplerType = SAMPLERTYPE_External;
+				else if (ValueStr.Contains(TEXT("Color"))) TexBase->SamplerType = SAMPLERTYPE_Color;
+				bPropertySet = true;
+			}
+			else if (Params->TryGetNumberField(TEXT("value"), ValueNum))
+			{
+				TexBase->SamplerType = static_cast<EMaterialSamplerType>(static_cast<uint8>(ValueNum));
+				bPropertySet = true;
+			}
+		}
+	}
+	else if (PropertyName.Equals(TEXT("SamplerSource"), ESearchCase::IgnoreCase))
+	{
+		if (UMaterialExpressionTextureSample* TexSample = Cast<UMaterialExpressionTextureSample>(Expression))
+		{
+			if (GetStringParam(Params, TEXT("value"), ValueStr, false))
+			{
+				// Map string names to enum values
+				if (ValueStr.Contains(TEXT("FromTextureAsset")) || ValueStr.Contains(TEXT("Asset"))) TexSample->SamplerSource = SSM_FromTextureAsset;
+				else if (ValueStr.Contains(TEXT("Wrap"))) TexSample->SamplerSource = SSM_Wrap_WorldGroupSettings;
+				else if (ValueStr.Contains(TEXT("Clamp"))) TexSample->SamplerSource = SSM_Clamp_WorldGroupSettings;
+				bPropertySet = true;
+			}
+			else if (Params->TryGetNumberField(TEXT("value"), ValueNum))
+			{
+				TexSample->SamplerSource = static_cast<ESamplerSourceMode>(static_cast<uint8>(ValueNum));
+				bPropertySet = true;
+			}
+		}
+	}
+	else if (PropertyName.Equals(TEXT("MipValueMode"), ESearchCase::IgnoreCase))
+	{
+		if (UMaterialExpressionTextureSample* TexSample = Cast<UMaterialExpressionTextureSample>(Expression))
+		{
+			if (GetStringParam(Params, TEXT("value"), ValueStr, false))
+			{
+				// Map string names to enum values
+				if (ValueStr.Contains(TEXT("None"))) TexSample->MipValueMode = TMVM_None;
+				else if (ValueStr.Contains(TEXT("MipLevel")) || ValueStr.Contains(TEXT("Level"))) TexSample->MipValueMode = TMVM_MipLevel;
+				else if (ValueStr.Contains(TEXT("MipBias")) || ValueStr.Contains(TEXT("Bias"))) TexSample->MipValueMode = TMVM_MipBias;
+				else if (ValueStr.Contains(TEXT("Derivative"))) TexSample->MipValueMode = TMVM_Derivative;
+				bPropertySet = true;
+			}
+			else if (Params->TryGetNumberField(TEXT("value"), ValueNum))
+			{
+				TexSample->MipValueMode = static_cast<ETextureMipValueMode>(static_cast<uint8>(ValueNum));
+				bPropertySet = true;
+			}
+		}
+	}
+	else if (PropertyName.Equals(TEXT("ConstCoordinate"), ESearchCase::IgnoreCase))
+	{
+		if (UMaterialExpressionTextureSample* TexSample = Cast<UMaterialExpressionTextureSample>(Expression))
+		{
+			if (Params->TryGetNumberField(TEXT("value"), ValueNum))
+			{
+				TexSample->ConstCoordinate = static_cast<uint8>(ValueNum);
+				bPropertySet = true;
+			}
+		}
+	}
+	else if (PropertyName.Equals(TEXT("ConstMipValue"), ESearchCase::IgnoreCase))
+	{
+		if (UMaterialExpressionTextureSample* TexSample = Cast<UMaterialExpressionTextureSample>(Expression))
+		{
+			if (Params->TryGetNumberField(TEXT("value"), ValueNum))
+			{
+				TexSample->ConstMipValue = static_cast<int32>(ValueNum);
+				bPropertySet = true;
+			}
+		}
+	}
+	else if (PropertyName.Equals(TEXT("AutomaticViewMipBias"), ESearchCase::IgnoreCase))
+	{
+		if (UMaterialExpressionTextureSample* TexSample = Cast<UMaterialExpressionTextureSample>(Expression))
+		{
+			if (Params->TryGetBoolField(TEXT("value"), ValueBool))
+			{
+				TexSample->AutomaticViewMipBias = ValueBool;
+				bPropertySet = true;
+			}
+		}
+	}
+	else if (PropertyName.Equals(TEXT("CoordinateIndex"), ESearchCase::IgnoreCase))
+	{
+		if (Params->TryGetNumberField(TEXT("value"), ValueNum))
+		{
+			if (UMaterialExpressionTextureCoordinate* TexCoord = Cast<UMaterialExpressionTextureCoordinate>(Expression))
+			{
+				TexCoord->CoordinateIndex = static_cast<int32>(ValueNum);
+				bPropertySet = true;
+			}
+		}
+	}
+	else if (PropertyName.Equals(TEXT("UTiling"), ESearchCase::IgnoreCase))
+	{
+		if (Params->TryGetNumberField(TEXT("value"), ValueNum))
+		{
+			if (UMaterialExpressionTextureCoordinate* TexCoord = Cast<UMaterialExpressionTextureCoordinate>(Expression))
+			{
+				TexCoord->UTiling = ValueNum;
+				bPropertySet = true;
+			}
+		}
+	}
+	else if (PropertyName.Equals(TEXT("VTiling"), ESearchCase::IgnoreCase))
+	{
+		if (Params->TryGetNumberField(TEXT("value"), ValueNum))
+		{
+			if (UMaterialExpressionTextureCoordinate* TexCoord = Cast<UMaterialExpressionTextureCoordinate>(Expression))
+			{
+				TexCoord->VTiling = ValueNum;
+				bPropertySet = true;
+			}
+		}
+	}
+	else if (PropertyName.Equals(TEXT("R"), ESearchCase::IgnoreCase) || 
+			 PropertyName.Equals(TEXT("G"), ESearchCase::IgnoreCase) ||
+			 PropertyName.Equals(TEXT("B"), ESearchCase::IgnoreCase) ||
+			 PropertyName.Equals(TEXT("A"), ESearchCase::IgnoreCase))
+	{
+		if (Params->TryGetBoolField(TEXT("value"), ValueBool))
+		{
+			if (UMaterialExpressionComponentMask* Mask = Cast<UMaterialExpressionComponentMask>(Expression))
+			{
+				if (PropertyName.Equals(TEXT("R"), ESearchCase::IgnoreCase)) Mask->R = ValueBool;
+				else if (PropertyName.Equals(TEXT("G"), ESearchCase::IgnoreCase)) Mask->G = ValueBool;
+				else if (PropertyName.Equals(TEXT("B"), ESearchCase::IgnoreCase)) Mask->B = ValueBool;
+				else if (PropertyName.Equals(TEXT("A"), ESearchCase::IgnoreCase)) Mask->A = ValueBool;
+				bPropertySet = true;
+			}
+			else if (UMaterialExpressionNoise* Noise = Cast<UMaterialExpressionNoise>(Expression))
+			{
+				if (PropertyName.Equals(TEXT("Turbulence"), ESearchCase::IgnoreCase)) Noise->bTurbulence = ValueBool;
+				else if (PropertyName.Equals(TEXT("Tiling"), ESearchCase::IgnoreCase)) Noise->bTiling = ValueBool;
+				bPropertySet = true;
+			}
+		}
+	}
+	
+	// Handle Noise-specific numeric properties
+	if (UMaterialExpressionNoise* Noise = Cast<UMaterialExpressionNoise>(Expression))
+	{
+		if (Params->TryGetNumberField(TEXT("value"), ValueNum))
+		{
+			if (PropertyName.Equals(TEXT("Scale"), ESearchCase::IgnoreCase))
+			{
+				Noise->Scale = ValueNum;
+				bPropertySet = true;
+			}
+			else if (PropertyName.Equals(TEXT("Quality"), ESearchCase::IgnoreCase))
+			{
+				Noise->Quality = FMath::Clamp(static_cast<int32>(ValueNum), 1, 4);
+				bPropertySet = true;
+			}
+			else if (PropertyName.Equals(TEXT("NoiseFunction"), ESearchCase::IgnoreCase))
+			{
+				Noise->NoiseFunction = static_cast<ENoiseFunction>(static_cast<uint8>(ValueNum));
+				bPropertySet = true;
+			}
+			else if (PropertyName.Equals(TEXT("Levels"), ESearchCase::IgnoreCase))
+			{
+				Noise->Levels = FMath::Clamp(static_cast<int32>(ValueNum), 1, 10);
+				bPropertySet = true;
+			}
+			else if (PropertyName.Equals(TEXT("OutputMin"), ESearchCase::IgnoreCase))
+			{
+				Noise->OutputMin = ValueNum;
+				bPropertySet = true;
+			}
+			else if (PropertyName.Equals(TEXT("OutputMax"), ESearchCase::IgnoreCase))
+			{
+				Noise->OutputMax = ValueNum;
+				bPropertySet = true;
+			}
+			else if (PropertyName.Equals(TEXT("LevelScale"), ESearchCase::IgnoreCase))
+			{
+				Noise->LevelScale = FMath::Clamp(static_cast<float>(ValueNum), 2.0f, 8.0f);
+				bPropertySet = true;
+			}
+			else if (PropertyName.Equals(TEXT("RepeatSize"), ESearchCase::IgnoreCase))
+			{
+				Noise->RepeatSize = FMath::Max(static_cast<uint32>(ValueNum), 4u);
+				bPropertySet = true;
+			}
+		}
+	}
+	
+	// Handle Custom expression properties
+	if (UMaterialExpressionCustom* Custom = Cast<UMaterialExpressionCustom>(Expression))
+	{
+		if (PropertyName.Equals(TEXT("Code"), ESearchCase::IgnoreCase))
+		{
+			if (GetStringParam(Params, TEXT("value"), ValueStr, false))
+			{
+				Custom->Code = ValueStr;
+				bPropertySet = true;
+			}
+		}
+		else if (PropertyName.Equals(TEXT("Description"), ESearchCase::IgnoreCase))
+		{
+			if (GetStringParam(Params, TEXT("value"), ValueStr, false))
+			{
+				Custom->Description = ValueStr;
+				bPropertySet = true;
+			}
+		}
+		else if (PropertyName.Equals(TEXT("OutputType"), ESearchCase::IgnoreCase))
+		{
+			if (GetStringParam(Params, TEXT("value"), ValueStr, false))
+			{
+				if (ValueStr.Contains(TEXT("Float1")) || ValueStr.Contains(TEXT("Scalar"))) Custom->OutputType = CMOT_Float1;
+				else if (ValueStr.Contains(TEXT("Float2"))) Custom->OutputType = CMOT_Float2;
+				else if (ValueStr.Contains(TEXT("Float3"))) Custom->OutputType = CMOT_Float3;
+				else if (ValueStr.Contains(TEXT("Float4"))) Custom->OutputType = CMOT_Float4;
+				else if (ValueStr.Contains(TEXT("MaterialAttributes"))) Custom->OutputType = CMOT_MaterialAttributes;
+				bPropertySet = true;
+			}
+			else if (Params->TryGetNumberField(TEXT("value"), ValueNum))
+			{
+				Custom->OutputType = static_cast<ECustomMaterialOutputType>(static_cast<uint8>(ValueNum));
+				bPropertySet = true;
+			}
+		}
+		else if (PropertyName.Equals(TEXT("Inputs"), ESearchCase::IgnoreCase) || PropertyName.Equals(TEXT("CustomInputs"), ESearchCase::IgnoreCase))
+		{
+			const TArray<TSharedPtr<FJsonValue>>* InputsArray;
+			if (Params->TryGetArrayField(TEXT("value"), InputsArray))
+			{
+				Custom->Inputs.Empty();
+				UE_LOG(LogTemp, Warning, TEXT("ECABridge: Setting Custom inputs, array size: %d"), InputsArray->Num());
+				for (const TSharedPtr<FJsonValue>& InputValue : *InputsArray)
+				{
+					const TSharedPtr<FJsonObject>* InputObj;
+					if (InputValue->TryGetObject(InputObj))
+					{
+						FCustomInput NewInput;
+						FString InputName;
+						if ((*InputObj)->TryGetStringField(TEXT("input_name"), InputName))
+						{
+							NewInput.InputName = FName(*InputName);
+							UE_LOG(LogTemp, Warning, TEXT("ECABridge: Set input name from object: '%s' -> FName: '%s'"), *InputName, *NewInput.InputName.ToString());
+						}
+						else
+						{
+							UE_LOG(LogTemp, Warning, TEXT("ECABridge: input_name field not found in object"));
+						}
+						Custom->Inputs.Add(NewInput);
+					}
+					else
+					{
+						// Simple string array - just input names
+						FString InputName = InputValue->AsString();
+						if (!InputName.IsEmpty())
+						{
+							FCustomInput NewInput;
+							NewInput.InputName = FName(*InputName);
+							UE_LOG(LogTemp, Warning, TEXT("ECABridge: Set input name from string: '%s' -> FName: '%s'"), *InputName, *NewInput.InputName.ToString());
+							Custom->Inputs.Add(NewInput);
+						}
+						else
+						{
+							UE_LOG(LogTemp, Warning, TEXT("ECABridge: Empty input name string"));
+						}
+					}
+				}
+				UE_LOG(LogTemp, Warning, TEXT("ECABridge: After setting, Custom->Inputs.Num() = %d"), Custom->Inputs.Num());
+				for (int32 i = 0; i < Custom->Inputs.Num(); i++)
+				{
+					UE_LOG(LogTemp, Warning, TEXT("ECABridge: Input[%d].InputName = '%s'"), i, *Custom->Inputs[i].InputName.ToString());
+				}
+				// Rebuild outputs and reconstruct graph node to ensure input pins are properly created
+				Custom->RebuildOutputs();
+				if (Custom->GraphNode)
+				{
+					Custom->GraphNode->ReconstructNode();
+				}
+				bPropertySet = true;
+			}
+		}
+		else if (PropertyName.Equals(TEXT("AdditionalOutputs"), ESearchCase::IgnoreCase))
+		{
+			const TArray<TSharedPtr<FJsonValue>>* OutputsArray;
+			if (Params->TryGetArrayField(TEXT("value"), OutputsArray))
+			{
+				Custom->AdditionalOutputs.Empty();
+				for (const TSharedPtr<FJsonValue>& OutputValue : *OutputsArray)
+				{
+					const TSharedPtr<FJsonObject>* OutputObj;
+					if (OutputValue->TryGetObject(OutputObj))
+					{
+						FCustomOutput NewOutput;
+						FString OutputName;
+						if ((*OutputObj)->TryGetStringField(TEXT("output_name"), OutputName))
+						{
+							NewOutput.OutputName = FName(*OutputName);
+						}
+						double OutputType;
+						if ((*OutputObj)->TryGetNumberField(TEXT("output_type"), OutputType))
+						{
+							NewOutput.OutputType = static_cast<ECustomMaterialOutputType>(static_cast<uint8>(OutputType));
+						}
+						Custom->AdditionalOutputs.Add(NewOutput);
+					}
+				}
+				// Rebuild outputs and reconstruct graph node to ensure output pins are properly created
+				Custom->RebuildOutputs();
+				if (Custom->GraphNode)
+				{
+					Custom->GraphNode->ReconstructNode();
+				}
+				bPropertySet = true;
+			}
+		}
+		else if (PropertyName.Equals(TEXT("AdditionalDefines"), ESearchCase::IgnoreCase))
+		{
+			const TArray<TSharedPtr<FJsonValue>>* DefinesArray;
+			if (Params->TryGetArrayField(TEXT("value"), DefinesArray))
+			{
+				Custom->AdditionalDefines.Empty();
+				for (const TSharedPtr<FJsonValue>& DefineValue : *DefinesArray)
+				{
+					const TSharedPtr<FJsonObject>* DefineObj;
+					if (DefineValue->TryGetObject(DefineObj))
+					{
+						FCustomDefine NewDefine;
+						(*DefineObj)->TryGetStringField(TEXT("define_name"), NewDefine.DefineName);
+						(*DefineObj)->TryGetStringField(TEXT("define_value"), NewDefine.DefineValue);
+						Custom->AdditionalDefines.Add(NewDefine);
+					}
+				}
+				bPropertySet = true;
+			}
+		}
+		else if (PropertyName.Equals(TEXT("IncludeFilePaths"), ESearchCase::IgnoreCase))
+		{
+			const TArray<TSharedPtr<FJsonValue>>* IncludesArray;
+			if (Params->TryGetArrayField(TEXT("value"), IncludesArray))
+			{
+				Custom->IncludeFilePaths.Empty();
+				for (const TSharedPtr<FJsonValue>& IncludeValue : *IncludesArray)
+				{
+					FString IncludePath = IncludeValue->AsString();
+					if (!IncludePath.IsEmpty())
+					{
+						Custom->IncludeFilePaths.Add(IncludePath);
+					}
+				}
+				bPropertySet = true;
+			}
+		}
+	}
+	
+	ClearMaterialCompilationErrors(Material);
+	Material->PostEditChange();
+	Material->MarkPackageDirty();
+	
+	// Refresh the material editor UI if open
+	RefreshMaterialEditorUI(Material);
+	
+	if (!bPropertySet)
+	{
+		return FECACommandResult::Error(FString::Printf(TEXT("Could not set property '%s' on this node type"), *PropertyName));
+	}
+	
+	TSharedPtr<FJsonObject> Result = MakeResult();
+	Result->SetStringField(TEXT("node_id"), NodeId);
+	Result->SetStringField(TEXT("property_name"), PropertyName);
+	
+	return FECACommandResult::Success(Result);
+}
+
+//------------------------------------------------------------------------------
+// BatchEditMaterialNodes
+//------------------------------------------------------------------------------
+
+FECACommandResult FECACommand_BatchEditMaterialNodes::Execute(const TSharedPtr<FJsonObject>& Params)
+{
+	FString MaterialPath;
+	if (!GetStringParam(Params, TEXT("material_path"), MaterialPath))
+	{
+		return FECACommandResult::Error(TEXT("Missing required parameter: material_path"));
+	}
+	
+	const TArray<TSharedPtr<FJsonValue>>* NodesArray;
+	if (!Params->TryGetArrayField(TEXT("nodes"), NodesArray))
+	{
+		return FECACommandResult::Error(TEXT("Missing required parameter: nodes"));
+	}
+	
+	UMaterial* Material = LoadMaterialByPath(MaterialPath);
+	if (!Material)
+	{
+		return FECACommandResult::Error(FString::Printf(TEXT("Material not found: %s"), *MaterialPath));
+	}
+	
+	Material->PreEditChange(nullptr);
+	
+	// Map of temp_id to created expressions
+	TMap<FString, UMaterialExpression*> TempIdMap;
+	TArray<TSharedPtr<FJsonValue>> CreatedNodesArray;
+	
+	// First pass: create all nodes
+	for (const TSharedPtr<FJsonValue>& NodeValue : *NodesArray)
+	{
+		const TSharedPtr<FJsonObject>* NodeObj;
+		if (!NodeValue->TryGetObject(NodeObj))
+		{
+			continue;
+		}
+		
+		FString TempId;
+		if (!(*NodeObj)->TryGetStringField(TEXT("temp_id"), TempId))
+		{
+			continue;
+		}
+		
+		FString NodeType;
+		if (!(*NodeObj)->TryGetStringField(TEXT("node_type"), NodeType))
+		{
+			continue;
+		}
+		
+		UClass* ExpressionClass = FindExpressionClass(NodeType);
+		if (!ExpressionClass)
+		{
+			ClearMaterialCompilationErrors(Material);
+			Material->PostEditChange();
+			return FECACommandResult::Error(FString::Printf(TEXT("Unknown material expression type: %s"), *NodeType));
+		}
+		
+		// Get position - auto-calculate if not provided
+		FVector2D Position(0, 0);
+		const TSharedPtr<FJsonObject>* PosObj;
+		if ((*NodeObj)->TryGetObjectField(TEXT("node_position"), PosObj))
+		{
+			Position.X = (*PosObj)->GetNumberField(TEXT("x"));
+			Position.Y = (*PosObj)->GetNumberField(TEXT("y"));
+		}
+		else
+		{
+			// Auto-calculate position for new node
+			Position = CalculateAutoMaterialNodePosition(Material);
+		}
+		
+		// Create expression
+		UMaterialExpression* NewExpression = NewObject<UMaterialExpression>(Material, ExpressionClass, NAME_None, RF_Transactional);
+		NewExpression->MaterialExpressionEditorX = Position.X;
+		NewExpression->MaterialExpressionEditorY = Position.Y;
+		
+		// Configure parameters
+		FString ParameterName;
+		if ((*NodeObj)->TryGetStringField(TEXT("parameter_name"), ParameterName))
+		{
+			if (UMaterialExpressionScalarParameter* ScalarParam = Cast<UMaterialExpressionScalarParameter>(NewExpression))
+			{
+				ScalarParam->ParameterName = FName(*ParameterName);
+			}
+			else if (UMaterialExpressionVectorParameter* VectorParam = Cast<UMaterialExpressionVectorParameter>(NewExpression))
+			{
+				VectorParam->ParameterName = FName(*ParameterName);
+			}
+			else if (UMaterialExpressionTextureSampleParameter2D* TexParam = Cast<UMaterialExpressionTextureSampleParameter2D>(NewExpression))
+			{
+				TexParam->ParameterName = FName(*ParameterName);
+			}
+		}
+		
+		// Set default value
+		const TSharedPtr<FJsonObject>* DefaultValueObj;
+		double DefaultValueNum;
+		FString DefaultValueStr;
+		
+		if ((*NodeObj)->TryGetObjectField(TEXT("default_value"), DefaultValueObj))
+		{
+			if (UMaterialExpressionVectorParameter* VectorParam = Cast<UMaterialExpressionVectorParameter>(NewExpression))
+			{
+				VectorParam->DefaultValue.R = (*DefaultValueObj)->GetNumberField(TEXT("r"));
+				VectorParam->DefaultValue.G = (*DefaultValueObj)->GetNumberField(TEXT("g"));
+				VectorParam->DefaultValue.B = (*DefaultValueObj)->GetNumberField(TEXT("b"));
+				if ((*DefaultValueObj)->HasField(TEXT("a")))
+					VectorParam->DefaultValue.A = (*DefaultValueObj)->GetNumberField(TEXT("a"));
+			}
+			else if (UMaterialExpressionConstant2Vector* Const2 = Cast<UMaterialExpressionConstant2Vector>(NewExpression))
+			{
+				Const2->R = (*DefaultValueObj)->GetNumberField(TEXT("r"));
+				Const2->G = (*DefaultValueObj)->GetNumberField(TEXT("g"));
+			}
+			else if (UMaterialExpressionConstant3Vector* Const3 = Cast<UMaterialExpressionConstant3Vector>(NewExpression))
+			{
+				Const3->Constant.R = (*DefaultValueObj)->GetNumberField(TEXT("r"));
+				Const3->Constant.G = (*DefaultValueObj)->GetNumberField(TEXT("g"));
+				Const3->Constant.B = (*DefaultValueObj)->GetNumberField(TEXT("b"));
+			}
+			else if (UMaterialExpressionConstant4Vector* Const4 = Cast<UMaterialExpressionConstant4Vector>(NewExpression))
+			{
+				Const4->Constant.R = (*DefaultValueObj)->GetNumberField(TEXT("r"));
+				Const4->Constant.G = (*DefaultValueObj)->GetNumberField(TEXT("g"));
+				Const4->Constant.B = (*DefaultValueObj)->GetNumberField(TEXT("b"));
+				Const4->Constant.A = (*DefaultValueObj)->HasField(TEXT("a")) ? (*DefaultValueObj)->GetNumberField(TEXT("a")) : 1.0f;
+			}
+		}
+		else if ((*NodeObj)->TryGetNumberField(TEXT("default_value"), DefaultValueNum))
+		{
+			if (UMaterialExpressionScalarParameter* ScalarParam = Cast<UMaterialExpressionScalarParameter>(NewExpression))
+			{
+				ScalarParam->DefaultValue = DefaultValueNum;
+			}
+			else if (UMaterialExpressionConstant* Const = Cast<UMaterialExpressionConstant>(NewExpression))
+			{
+				Const->R = DefaultValueNum;
+			}
+		}
+		else if ((*NodeObj)->TryGetStringField(TEXT("default_value"), DefaultValueStr))
+		{
+			if (UMaterialExpressionTextureSample* TexSample = Cast<UMaterialExpressionTextureSample>(NewExpression))
+			{
+				UTexture2D* Texture = LoadObject<UTexture2D>(nullptr, *DefaultValueStr);
+				if (Texture)
+				{
+					TexSample->Texture = Texture;
+				}
+			}
+		}
+		
+		// Handle TextureSample properties
+		if (UMaterialExpressionTextureSample* TexSample = Cast<UMaterialExpressionTextureSample>(NewExpression))
+		{
+			double TexValue;
+			FString TexStr;
+			bool TexBool;
+			
+			// SamplerType
+			if ((*NodeObj)->TryGetStringField(TEXT("sampler_type"), TexStr))
+			{
+				if (TexStr.Contains(TEXT("LinearColor"))) TexSample->SamplerType = SAMPLERTYPE_LinearColor;
+				else if (TexStr.Contains(TEXT("LinearGrayscale"))) TexSample->SamplerType = SAMPLERTYPE_LinearGrayscale;
+				else if (TexStr.Contains(TEXT("Grayscale"))) TexSample->SamplerType = SAMPLERTYPE_Grayscale;
+				else if (TexStr.Contains(TEXT("Normal"))) TexSample->SamplerType = SAMPLERTYPE_Normal;
+				else if (TexStr.Contains(TEXT("Alpha"))) TexSample->SamplerType = SAMPLERTYPE_Alpha;
+				else if (TexStr.Contains(TEXT("Masks"))) TexSample->SamplerType = SAMPLERTYPE_Masks;
+				else if (TexStr.Contains(TEXT("Data"))) TexSample->SamplerType = SAMPLERTYPE_Data;
+				else if (TexStr.Contains(TEXT("External"))) TexSample->SamplerType = SAMPLERTYPE_External;
+				else if (TexStr.Contains(TEXT("Color"))) TexSample->SamplerType = SAMPLERTYPE_Color;
+			}
+			else if ((*NodeObj)->TryGetNumberField(TEXT("sampler_type"), TexValue))
+			{
+				TexSample->SamplerType = static_cast<EMaterialSamplerType>(static_cast<uint8>(TexValue));
+			}
+			
+			// SamplerSource
+			if ((*NodeObj)->TryGetStringField(TEXT("sampler_source"), TexStr))
+			{
+				if (TexStr.Contains(TEXT("FromTextureAsset")) || TexStr.Contains(TEXT("Asset"))) TexSample->SamplerSource = SSM_FromTextureAsset;
+				else if (TexStr.Contains(TEXT("Wrap"))) TexSample->SamplerSource = SSM_Wrap_WorldGroupSettings;
+				else if (TexStr.Contains(TEXT("Clamp"))) TexSample->SamplerSource = SSM_Clamp_WorldGroupSettings;
+			}
+			else if ((*NodeObj)->TryGetNumberField(TEXT("sampler_source"), TexValue))
+			{
+				TexSample->SamplerSource = static_cast<ESamplerSourceMode>(static_cast<uint8>(TexValue));
+			}
+			
+			// MipValueMode
+			if ((*NodeObj)->TryGetStringField(TEXT("mip_value_mode"), TexStr))
+			{
+				if (TexStr.Contains(TEXT("None"))) TexSample->MipValueMode = TMVM_None;
+				else if (TexStr.Contains(TEXT("MipLevel")) || TexStr.Contains(TEXT("Level"))) TexSample->MipValueMode = TMVM_MipLevel;
+				else if (TexStr.Contains(TEXT("MipBias")) || TexStr.Contains(TEXT("Bias"))) TexSample->MipValueMode = TMVM_MipBias;
+				else if (TexStr.Contains(TEXT("Derivative"))) TexSample->MipValueMode = TMVM_Derivative;
+			}
+			else if ((*NodeObj)->TryGetNumberField(TEXT("mip_value_mode"), TexValue))
+			{
+				TexSample->MipValueMode = static_cast<ETextureMipValueMode>(static_cast<uint8>(TexValue));
+			}
+			
+			// Other numeric/bool properties
+			if ((*NodeObj)->TryGetNumberField(TEXT("const_coordinate"), TexValue))
+				TexSample->ConstCoordinate = static_cast<uint8>(TexValue);
+			if ((*NodeObj)->TryGetNumberField(TEXT("const_mip_value"), TexValue))
+				TexSample->ConstMipValue = static_cast<int32>(TexValue);
+			if ((*NodeObj)->TryGetBoolField(TEXT("automatic_view_mip_bias"), TexBool))
+				TexSample->AutomaticViewMipBias = TexBool;
+		}
+		
+		// Handle Noise properties
+		if (UMaterialExpressionNoise* Noise = Cast<UMaterialExpressionNoise>(NewExpression))
+		{
+			double NoiseValue;
+			bool NoiseBool;
+			if ((*NodeObj)->TryGetNumberField(TEXT("scale"), NoiseValue))
+				Noise->Scale = NoiseValue;
+			if ((*NodeObj)->TryGetNumberField(TEXT("quality"), NoiseValue))
+				Noise->Quality = FMath::Clamp(static_cast<int32>(NoiseValue), 1, 4);
+			if ((*NodeObj)->TryGetNumberField(TEXT("noise_function"), NoiseValue))
+				Noise->NoiseFunction = static_cast<ENoiseFunction>(static_cast<uint8>(NoiseValue));
+			if ((*NodeObj)->TryGetBoolField(TEXT("turbulence"), NoiseBool))
+				Noise->bTurbulence = NoiseBool;
+			if ((*NodeObj)->TryGetNumberField(TEXT("levels"), NoiseValue))
+				Noise->Levels = FMath::Clamp(static_cast<int32>(NoiseValue), 1, 10);
+			if ((*NodeObj)->TryGetNumberField(TEXT("output_min"), NoiseValue))
+				Noise->OutputMin = NoiseValue;
+			if ((*NodeObj)->TryGetNumberField(TEXT("output_max"), NoiseValue))
+				Noise->OutputMax = NoiseValue;
+			if ((*NodeObj)->TryGetNumberField(TEXT("level_scale"), NoiseValue))
+				Noise->LevelScale = FMath::Clamp(static_cast<float>(NoiseValue), 2.0f, 8.0f);
+			if ((*NodeObj)->TryGetBoolField(TEXT("tiling"), NoiseBool))
+				Noise->bTiling = NoiseBool;
+			if ((*NodeObj)->TryGetNumberField(TEXT("repeat_size"), NoiseValue))
+				Noise->RepeatSize = FMath::Max(static_cast<uint32>(NoiseValue), 4u);
+		}
+		
+		// Handle Custom expression properties
+		if (UMaterialExpressionCustom* Custom = Cast<UMaterialExpressionCustom>(NewExpression))
+		{
+			FString CustomStr;
+			double CustomNum;
+			
+			if ((*NodeObj)->TryGetStringField(TEXT("code"), CustomStr))
+				Custom->Code = CustomStr;
+			if ((*NodeObj)->TryGetStringField(TEXT("description"), CustomStr))
+				Custom->Description = CustomStr;
+			
+			// OutputType
+			if ((*NodeObj)->TryGetStringField(TEXT("output_type"), CustomStr))
+			{
+				if (CustomStr.Contains(TEXT("Float1")) || CustomStr.Contains(TEXT("Scalar"))) Custom->OutputType = CMOT_Float1;
+				else if (CustomStr.Contains(TEXT("Float2"))) Custom->OutputType = CMOT_Float2;
+				else if (CustomStr.Contains(TEXT("Float3"))) Custom->OutputType = CMOT_Float3;
+				else if (CustomStr.Contains(TEXT("Float4"))) Custom->OutputType = CMOT_Float4;
+				else if (CustomStr.Contains(TEXT("MaterialAttributes"))) Custom->OutputType = CMOT_MaterialAttributes;
+			}
+			else if ((*NodeObj)->TryGetNumberField(TEXT("output_type"), CustomNum))
+			{
+				Custom->OutputType = static_cast<ECustomMaterialOutputType>(static_cast<uint8>(CustomNum));
+			}
+			
+			// Inputs array
+			const TArray<TSharedPtr<FJsonValue>>* InputsArray;
+			if ((*NodeObj)->TryGetArrayField(TEXT("inputs"), InputsArray))
+			{
+				Custom->Inputs.Empty();
+				for (const TSharedPtr<FJsonValue>& InputValue : *InputsArray)
+				{
+					const TSharedPtr<FJsonObject>* InputObj;
+					if (InputValue->TryGetObject(InputObj))
+					{
+						FCustomInput NewInput;
+						FString InputName;
+						if ((*InputObj)->TryGetStringField(TEXT("input_name"), InputName))
+						{
+							NewInput.InputName = FName(*InputName);
+						}
+						Custom->Inputs.Add(NewInput);
+					}
+					else
+					{
+						FString InputName = InputValue->AsString();
+						if (!InputName.IsEmpty())
+						{
+							FCustomInput NewInput;
+							NewInput.InputName = FName(*InputName);
+							Custom->Inputs.Add(NewInput);
+						}
+					}
+				}
+			}
+			
+			// Additional outputs
+			const TArray<TSharedPtr<FJsonValue>>* OutputsArray;
+			if ((*NodeObj)->TryGetArrayField(TEXT("additional_outputs"), OutputsArray))
+			{
+				Custom->AdditionalOutputs.Empty();
+				for (const TSharedPtr<FJsonValue>& OutputValue : *OutputsArray)
+				{
+					const TSharedPtr<FJsonObject>* OutputObj;
+					if (OutputValue->TryGetObject(OutputObj))
+					{
+						FCustomOutput NewOutput;
+						FString OutputName;
+						if ((*OutputObj)->TryGetStringField(TEXT("output_name"), OutputName))
+						{
+							NewOutput.OutputName = FName(*OutputName);
+						}
+						double OutputType;
+						if ((*OutputObj)->TryGetNumberField(TEXT("output_type"), OutputType))
+						{
+							NewOutput.OutputType = static_cast<ECustomMaterialOutputType>(static_cast<uint8>(OutputType));
+						}
+						Custom->AdditionalOutputs.Add(NewOutput);
+					}
+				}
+			}
+			
+			// Additional defines
+			const TArray<TSharedPtr<FJsonValue>>* DefinesArray;
+			if ((*NodeObj)->TryGetArrayField(TEXT("additional_defines"), DefinesArray))
+			{
+				Custom->AdditionalDefines.Empty();
+				for (const TSharedPtr<FJsonValue>& DefineValue : *DefinesArray)
+				{
+					const TSharedPtr<FJsonObject>* DefineObj;
+					if (DefineValue->TryGetObject(DefineObj))
+					{
+						FCustomDefine NewDefine;
+						(*DefineObj)->TryGetStringField(TEXT("define_name"), NewDefine.DefineName);
+						(*DefineObj)->TryGetStringField(TEXT("define_value"), NewDefine.DefineValue);
+						Custom->AdditionalDefines.Add(NewDefine);
+					}
+				}
+			}
+			
+			// Include file paths
+			const TArray<TSharedPtr<FJsonValue>>* IncludesArray;
+			if ((*NodeObj)->TryGetArrayField(TEXT("include_file_paths"), IncludesArray))
+			{
+				Custom->IncludeFilePaths.Empty();
+				for (const TSharedPtr<FJsonValue>& IncludeValue : *IncludesArray)
+				{
+					FString IncludePath = IncludeValue->AsString();
+					if (!IncludePath.IsEmpty())
+					{
+						Custom->IncludeFilePaths.Add(IncludePath);
+					}
+				}
+			}
+			
+			// Rebuild outputs and reconstruct graph node to ensure input/output pins are properly created
+			Custom->RebuildOutputs();
+			if (Custom->GraphNode)
+			{
+				Custom->GraphNode->ReconstructNode();
+			}
+		}
+		
+		// Add to material
+		Material->GetEditorOnlyData()->ExpressionCollection.AddExpression(NewExpression);
+		TempIdMap.Add(TempId, NewExpression);
+		
+		// Add to result
+		TSharedPtr<FJsonObject> CreatedNode = MakeShared<FJsonObject>();
+		CreatedNode->SetStringField(TEXT("temp_id"), TempId);
+		CreatedNode->SetStringField(TEXT("node_id"), NewExpression->MaterialExpressionGuid.ToString());
+		CreatedNode->SetStringField(TEXT("node_type"), NodeType);
+		CreatedNodesArray.Add(MakeShared<FJsonValueObject>(CreatedNode));
+	}
+	
+	// Second pass: create connections
+	const TArray<TSharedPtr<FJsonValue>>* ConnectionsArray;
+	TArray<TSharedPtr<FJsonValue>> ConnectionResultsArray;
+	TArray<TSharedPtr<FJsonValue>> FailedConnectionsArray;
+	int32 SuccessfulConnections = 0;
+	
+	if (Params->TryGetArrayField(TEXT("connections"), ConnectionsArray))
+	{
+		for (const TSharedPtr<FJsonValue>& ConnValue : *ConnectionsArray)
+		{
+			const TSharedPtr<FJsonObject>* ConnObj;
+			if (!ConnValue->TryGetObject(ConnObj))
+			{
+				continue;
+			}
+			
+			FString SourceNode, SourceOutput, TargetNode, TargetInput;
+			(*ConnObj)->TryGetStringField(TEXT("source_node"), SourceNode);
+			(*ConnObj)->TryGetStringField(TEXT("source_output"), SourceOutput);
+			(*ConnObj)->TryGetStringField(TEXT("target_node"), TargetNode);
+			(*ConnObj)->TryGetStringField(TEXT("target_input"), TargetInput);
+			
+			if (SourceOutput.IsEmpty()) SourceOutput = TEXT("0");
+			
+			// Resolve source expression
+			UMaterialExpression* SourceExpr = nullptr;
+			if (UMaterialExpression** Found = TempIdMap.Find(SourceNode))
+			{
+				SourceExpr = *Found;
+			}
+			else
+			{
+				FGuid SourceGuid;
+				if (FGuid::Parse(SourceNode, SourceGuid))
+				{
+					SourceExpr = FindExpressionByGuid(Material, SourceGuid);
+				}
+			}
+			
+			if (!SourceExpr)
+			{
+				TSharedPtr<FJsonObject> FailedConn = MakeShared<FJsonObject>();
+				FailedConn->SetStringField(TEXT("source_node"), SourceNode);
+				FailedConn->SetStringField(TEXT("target_node"), TargetNode);
+				FailedConn->SetStringField(TEXT("target_input"), TargetInput);
+				FailedConn->SetStringField(TEXT("error"), FString::Printf(TEXT("Source node not found: %s"), *SourceNode));
+				FailedConnectionsArray.Add(MakeShared<FJsonValueObject>(FailedConn));
+				continue;
+			}
+			
+			int32 OutputIndex = ParseOutputIndex(SourceOutput, SourceExpr);
+			
+			// Connect to material or expression
+			if (TargetNode.Equals(TEXT("material"), ESearchCase::IgnoreCase))
+			{
+				FExpressionInput* MaterialInput = GetMaterialInput(Material, TargetInput);
+				if (MaterialInput)
+				{
+					MaterialInput->Expression = SourceExpr;
+					MaterialInput->OutputIndex = OutputIndex;
+					SuccessfulConnections++;
+				}
+				else
+				{
+					TSharedPtr<FJsonObject> FailedConn = MakeShared<FJsonObject>();
+					FailedConn->SetStringField(TEXT("source_node"), SourceNode);
+					FailedConn->SetStringField(TEXT("target_node"), TargetNode);
+					FailedConn->SetStringField(TEXT("target_input"), TargetInput);
+					FailedConn->SetStringField(TEXT("error"), FString::Printf(TEXT("Unknown material input: %s"), *TargetInput));
+					FailedConnectionsArray.Add(MakeShared<FJsonValueObject>(FailedConn));
+				}
+			}
+			else
+			{
+				// Resolve target expression
+				UMaterialExpression* TargetExpr = nullptr;
+				if (UMaterialExpression** Found = TempIdMap.Find(TargetNode))
+				{
+					TargetExpr = *Found;
+				}
+				else
+				{
+					FGuid TargetGuid;
+					if (FGuid::Parse(TargetNode, TargetGuid))
+					{
+						TargetExpr = FindExpressionByGuid(Material, TargetGuid);
+					}
+				}
+				
+				if (!TargetExpr)
+				{
+					TSharedPtr<FJsonObject> FailedConn = MakeShared<FJsonObject>();
+					FailedConn->SetStringField(TEXT("source_node"), SourceNode);
+					FailedConn->SetStringField(TEXT("target_node"), TargetNode);
+					FailedConn->SetStringField(TEXT("target_input"), TargetInput);
+					FailedConn->SetStringField(TEXT("error"), FString::Printf(TEXT("Target node not found: %s"), *TargetNode));
+					FailedConnectionsArray.Add(MakeShared<FJsonValueObject>(FailedConn));
+					continue;
+				}
+				
+				{
+					FExpressionInput* TargetInputPtr = nullptr;
+					
+					if (TargetInput.IsNumeric())
+					{
+						int32 InputIndex = FCString::Atoi(*TargetInput);
+						TargetInputPtr = TargetExpr->GetInput(InputIndex);
+					}
+					else
+					{
+						// Try by name - iterate through all inputs
+						for (int32 i = 0; ; i++)
+						{
+							FExpressionInput* Input = TargetExpr->GetInput(i);
+							if (!Input)
+							{
+								break;
+							}
+							FName InputName = TargetExpr->GetInputName(i);
+							if (InputName.ToString().Equals(TargetInput, ESearchCase::IgnoreCase))
+							{
+								TargetInputPtr = Input;
+								break;
+							}
+						}
+						
+						// Handle common named inputs
+						if (!TargetInputPtr)
+						{
+							if (UMaterialExpressionMultiply* Mult = Cast<UMaterialExpressionMultiply>(TargetExpr))
+							{
+								if (TargetInput.Equals(TEXT("A"), ESearchCase::IgnoreCase)) TargetInputPtr = &Mult->A;
+								else if (TargetInput.Equals(TEXT("B"), ESearchCase::IgnoreCase)) TargetInputPtr = &Mult->B;
+							}
+							else if (UMaterialExpressionAdd* Add = Cast<UMaterialExpressionAdd>(TargetExpr))
+							{
+								if (TargetInput.Equals(TEXT("A"), ESearchCase::IgnoreCase)) TargetInputPtr = &Add->A;
+								else if (TargetInput.Equals(TEXT("B"), ESearchCase::IgnoreCase)) TargetInputPtr = &Add->B;
+							}
+							else if (UMaterialExpressionSubtract* Sub = Cast<UMaterialExpressionSubtract>(TargetExpr))
+							{
+								if (TargetInput.Equals(TEXT("A"), ESearchCase::IgnoreCase)) TargetInputPtr = &Sub->A;
+								else if (TargetInput.Equals(TEXT("B"), ESearchCase::IgnoreCase)) TargetInputPtr = &Sub->B;
+							}
+							else if (UMaterialExpressionDivide* Div = Cast<UMaterialExpressionDivide>(TargetExpr))
+							{
+								if (TargetInput.Equals(TEXT("A"), ESearchCase::IgnoreCase)) TargetInputPtr = &Div->A;
+								else if (TargetInput.Equals(TEXT("B"), ESearchCase::IgnoreCase)) TargetInputPtr = &Div->B;
+							}
+							else if (UMaterialExpressionPower* Pow = Cast<UMaterialExpressionPower>(TargetExpr))
+							{
+								if (TargetInput.Equals(TEXT("Base"), ESearchCase::IgnoreCase)) TargetInputPtr = &Pow->Base;
+								else if (TargetInput.Equals(TEXT("Exponent"), ESearchCase::IgnoreCase) || TargetInput.Equals(TEXT("Exp"), ESearchCase::IgnoreCase)) TargetInputPtr = &Pow->Exponent;
+							}
+							else if (UMaterialExpressionClamp* Clamp = Cast<UMaterialExpressionClamp>(TargetExpr))
+							{
+								if (TargetInput.Equals(TEXT("Input"), ESearchCase::IgnoreCase) || TargetInput.Equals(TEXT("0"), ESearchCase::IgnoreCase)) TargetInputPtr = &Clamp->Input;
+								else if (TargetInput.Equals(TEXT("Min"), ESearchCase::IgnoreCase)) TargetInputPtr = &Clamp->Min;
+								else if (TargetInput.Equals(TEXT("Max"), ESearchCase::IgnoreCase)) TargetInputPtr = &Clamp->Max;
+							}
+							else if (UMaterialExpressionIf* If = Cast<UMaterialExpressionIf>(TargetExpr))
+							{
+								if (TargetInput.Equals(TEXT("A"), ESearchCase::IgnoreCase)) TargetInputPtr = &If->A;
+								else if (TargetInput.Equals(TEXT("B"), ESearchCase::IgnoreCase)) TargetInputPtr = &If->B;
+								else if (TargetInput.Equals(TEXT("AGreaterThanB"), ESearchCase::IgnoreCase) || TargetInput.Equals(TEXT("A>B"), ESearchCase::IgnoreCase)) TargetInputPtr = &If->AGreaterThanB;
+								else if (TargetInput.Equals(TEXT("AEqualsB"), ESearchCase::IgnoreCase) || TargetInput.Equals(TEXT("A==B"), ESearchCase::IgnoreCase) || TargetInput.Equals(TEXT("A=B"), ESearchCase::IgnoreCase)) TargetInputPtr = &If->AEqualsB;
+								else if (TargetInput.Equals(TEXT("ALessThanB"), ESearchCase::IgnoreCase) || TargetInput.Equals(TEXT("A<B"), ESearchCase::IgnoreCase)) TargetInputPtr = &If->ALessThanB;
+							}
+							else if (UMaterialExpressionLinearInterpolate* Lerp = Cast<UMaterialExpressionLinearInterpolate>(TargetExpr))
+							{
+								if (TargetInput.Equals(TEXT("A"), ESearchCase::IgnoreCase)) TargetInputPtr = &Lerp->A;
+								else if (TargetInput.Equals(TEXT("B"), ESearchCase::IgnoreCase)) TargetInputPtr = &Lerp->B;
+								else if (TargetInput.Equals(TEXT("Alpha"), ESearchCase::IgnoreCase)) TargetInputPtr = &Lerp->Alpha;
+							}
+						}
+					}
+					
+					if (TargetInputPtr)
+					{
+						TargetInputPtr->Expression = SourceExpr;
+						TargetInputPtr->OutputIndex = OutputIndex;
+						SuccessfulConnections++;
+					}
+					else
+					{
+						TSharedPtr<FJsonObject> FailedConn = MakeShared<FJsonObject>();
+						FailedConn->SetStringField(TEXT("source_node"), SourceNode);
+						FailedConn->SetStringField(TEXT("target_node"), TargetNode);
+						FailedConn->SetStringField(TEXT("target_input"), TargetInput);
+						FailedConn->SetStringField(TEXT("error"), FString::Printf(TEXT("Input '%s' not found on target node"), *TargetInput));
+						FailedConnectionsArray.Add(MakeShared<FJsonValueObject>(FailedConn));
+					}
+				}
+			}
+		}
+	}
+	
+	ClearMaterialCompilationErrors(Material);
+	Material->PostEditChange();
+	Material->MarkPackageDirty();
+	
+	// Refresh the material editor UI if open
+	RefreshMaterialEditorUI(Material);
+	
+	TSharedPtr<FJsonObject> Result = MakeResult();
+	Result->SetArrayField(TEXT("created_nodes"), CreatedNodesArray);
+	Result->SetNumberField(TEXT("node_count"), CreatedNodesArray.Num());
+	Result->SetNumberField(TEXT("connections_made"), SuccessfulConnections);
+	
+	// Include failed connections if any
+	if (FailedConnectionsArray.Num() > 0)
+	{
+		Result->SetArrayField(TEXT("failed_connections"), FailedConnectionsArray);
+		Result->SetNumberField(TEXT("failed_connection_count"), FailedConnectionsArray.Num());
+	}
+	
+	// Include any compilation errors
+	AddCompilationErrorsToResult(Result, Material);
+	
+	return FECACommandResult::Success(Result);
+}
+
+//------------------------------------------------------------------------------
+// ListMaterialExpressionTypes
+//------------------------------------------------------------------------------
+
+FECACommandResult FECACommand_ListMaterialExpressionTypes::Execute(const TSharedPtr<FJsonObject>& Params)
+{
+	FString CategoryFilter;
+	GetStringParam(Params, TEXT("category_filter"), CategoryFilter, false);
+	
+	FString SearchFilter;
+	GetStringParam(Params, TEXT("search"), SearchFilter, false);
+	
+	TArray<TSharedPtr<FJsonValue>> TypesArray;
+	
+	for (TObjectIterator<UClass> It; It; ++It)
+	{
+		if (It->IsChildOf(UMaterialExpression::StaticClass()) && !It->HasAnyClassFlags(CLASS_Abstract))
+		{
+			FString ClassName = It->GetName();
+			FString SimpleName = ClassName.Replace(TEXT("MaterialExpression"), TEXT(""));
+			
+			// Apply search filter
+			if (!SearchFilter.IsEmpty())
+			{
+				if (!SimpleName.Contains(SearchFilter, ESearchCase::IgnoreCase) &&
+					!ClassName.Contains(SearchFilter, ESearchCase::IgnoreCase))
+				{
+					continue;
+				}
+			}
+			
+			// Determine category
+			FString Category = TEXT("Other");
+			if (SimpleName.Contains(TEXT("Parameter"))) Category = TEXT("Parameters");
+			else if (SimpleName.Contains(TEXT("Texture")) || SimpleName.Contains(TEXT("Sample"))) Category = TEXT("Texture");
+			else if (SimpleName.Contains(TEXT("Constant"))) Category = TEXT("Constants");
+			else if (SimpleName.Contains(TEXT("Add")) || SimpleName.Contains(TEXT("Subtract")) || 
+					 SimpleName.Contains(TEXT("Multiply")) || SimpleName.Contains(TEXT("Divide")) ||
+					 SimpleName.Contains(TEXT("Power")) || SimpleName.Contains(TEXT("Abs")) ||
+					 SimpleName.Contains(TEXT("Clamp")) || SimpleName.Contains(TEXT("Lerp")) ||
+					 SimpleName.Contains(TEXT("Dot")) || SimpleName.Contains(TEXT("Cross"))) Category = TEXT("Math");
+			else if (SimpleName.Contains(TEXT("Mask")) || SimpleName.Contains(TEXT("Append")) ||
+					 SimpleName.Contains(TEXT("Component"))) Category = TEXT("Utility");
+			else if (SimpleName.Contains(TEXT("World")) || SimpleName.Contains(TEXT("Camera")) ||
+					 SimpleName.Contains(TEXT("Position")) || SimpleName.Contains(TEXT("Normal"))) Category = TEXT("Coordinates");
+			else if (SimpleName.Contains(TEXT("Time")) || SimpleName.Contains(TEXT("Panner"))) Category = TEXT("Utility");
+			
+			// Apply category filter
+			if (!CategoryFilter.IsEmpty() && !Category.Equals(CategoryFilter, ESearchCase::IgnoreCase))
+			{
+				continue;
+			}
+			
+			TSharedPtr<FJsonObject> TypeInfo = MakeShared<FJsonObject>();
+			TypeInfo->SetStringField(TEXT("name"), SimpleName);
+			TypeInfo->SetStringField(TEXT("class"), ClassName);
+			TypeInfo->SetStringField(TEXT("category"), Category);
+			
+			TypesArray.Add(MakeShared<FJsonValueObject>(TypeInfo));
+		}
+	}
+	
+	TSharedPtr<FJsonObject> Result = MakeResult();
+	Result->SetArrayField(TEXT("expression_types"), TypesArray);
+	Result->SetNumberField(TEXT("count"), TypesArray.Num());
+	
+	return FECACommandResult::Success(Result);
+}
+
+//------------------------------------------------------------------------------
+// GetMaterialErrors
+//------------------------------------------------------------------------------
+
+FECACommandResult FECACommand_GetMaterialErrors::Execute(const TSharedPtr<FJsonObject>& Params)
+{
+	FString MaterialPath;
+	if (!GetStringParam(Params, TEXT("material_path"), MaterialPath))
+	{
+		return FECACommandResult::Error(TEXT("Missing required parameter: material_path"));
+	}
+	
+	UMaterial* Material = LoadMaterialByPath(MaterialPath);
+	if (!Material)
+	{
+		return FECACommandResult::Error(FString::Printf(TEXT("Material not found: %s"), *MaterialPath));
+	}
+	
+	TSharedPtr<FJsonObject> Result = MakeResult();
+	Result->SetStringField(TEXT("material_path"), MaterialPath);
+	
+	// Get compilation errors
+	TArray<TSharedPtr<FJsonValue>> Errors = GetMaterialCompilationErrors(Material);
+	Result->SetArrayField(TEXT("errors"), Errors);
+	Result->SetNumberField(TEXT("error_count"), Errors.Num());
+	Result->SetBoolField(TEXT("has_errors"), Errors.Num() > 0);
+	
+	return FECACommandResult::Success(Result);
+}
+
+//------------------------------------------------------------------------------
+// SetCustomNodeInputName
+//------------------------------------------------------------------------------
+
+FECACommandResult FECACommand_SetCustomNodeInputName::Execute(const TSharedPtr<FJsonObject>& Params)
+{
+	FString MaterialPath;
+	if (!GetStringParam(Params, TEXT("material_path"), MaterialPath))
+	{
+		return FECACommandResult::Error(TEXT("Missing required parameter: material_path"));
+	}
+	
+	FString NodeId;
+	if (!GetStringParam(Params, TEXT("node_id"), NodeId))
+	{
+		return FECACommandResult::Error(TEXT("Missing required parameter: node_id"));
+	}
+	
+	double InputIndexDouble;
+	if (!Params->TryGetNumberField(TEXT("input_index"), InputIndexDouble))
+	{
+		return FECACommandResult::Error(TEXT("Missing required parameter: input_index"));
+	}
+	int32 InputIndex = static_cast<int32>(InputIndexDouble);
+	
+	FString InputName;
+	if (!GetStringParam(Params, TEXT("input_name"), InputName))
+	{
+		return FECACommandResult::Error(TEXT("Missing required parameter: input_name"));
+	}
+	
+	UMaterial* Material = LoadMaterialByPath(MaterialPath);
+	if (!Material)
+	{
+		return FECACommandResult::Error(FString::Printf(TEXT("Material not found: %s"), *MaterialPath));
+	}
+	
+	FGuid NodeGuid;
+	if (!FGuid::Parse(NodeId, NodeGuid))
+	{
+		return FECACommandResult::Error(FString::Printf(TEXT("Invalid node GUID: %s"), *NodeId));
+	}
+	
+	UMaterialExpression* Expression = FindExpressionByGuid(Material, NodeGuid);
+	if (!Expression)
+	{
+		return FECACommandResult::Error(FString::Printf(TEXT("Node not found: %s"), *NodeId));
+	}
+	
+	UMaterialExpressionCustom* Custom = Cast<UMaterialExpressionCustom>(Expression);
+	if (!Custom)
+	{
+		return FECACommandResult::Error(FString::Printf(TEXT("Node is not a Custom expression: %s"), *NodeId));
+	}
+	
+	if (InputIndex < 0)
+	{
+		return FECACommandResult::Error(FString::Printf(TEXT("Invalid input index: %d (must be >= 0)"), InputIndex));
+	}
+	
+	Material->PreEditChange(nullptr);
+	
+	// Expand the Inputs array if necessary
+	while (Custom->Inputs.Num() <= InputIndex)
+	{
+		FCustomInput NewInput;
+		NewInput.InputName = FName(*FString::Printf(TEXT("Input%d"), Custom->Inputs.Num()));
+		Custom->Inputs.Add(NewInput);
+	}
+	
+	// Strip spaces from input name (same as PostEditChangeProperty does)
+	FString CleanInputName = InputName;
+	CleanInputName.ReplaceInline(TEXT(" "), TEXT(""));
+	
+	// Set the input name
+	Custom->Inputs[InputIndex].InputName = FName(*CleanInputName);
+	
+	UE_LOG(LogTemp, Warning, TEXT("ECABridge SetCustomNodeInputName: Set Input[%d].InputName = '%s' (cleaned: '%s')"), 
+		InputIndex, *InputName, *Custom->Inputs[InputIndex].InputName.ToString());
+	
+	// Rebuild outputs to ensure pins are updated
+	Custom->RebuildOutputs();
+	
+	// Reconstruct the graph node if it exists (material editor is open)
+	if (Custom->GraphNode)
+	{
+		Custom->GraphNode->ReconstructNode();
+	}
+	
+	ClearMaterialCompilationErrors(Material);
+	Material->PostEditChange();
+	Material->MarkPackageDirty();
+	
+	// Refresh the material editor UI if open
+	RefreshMaterialEditorUI(Material);
+	
+	TSharedPtr<FJsonObject> Result = MakeResult();
+	Result->SetStringField(TEXT("node_id"), NodeId);
+	Result->SetNumberField(TEXT("input_index"), InputIndex);
+	Result->SetStringField(TEXT("input_name"), Custom->Inputs[InputIndex].InputName.ToString());
+	Result->SetNumberField(TEXT("total_inputs"), Custom->Inputs.Num());
+	
+	// Return all input names for verification
+	TArray<TSharedPtr<FJsonValue>> InputNamesArray;
+	for (int32 i = 0; i < Custom->Inputs.Num(); i++)
+	{
+		TSharedPtr<FJsonObject> InputObj = MakeShared<FJsonObject>();
+		InputObj->SetNumberField(TEXT("index"), i);
+		InputObj->SetStringField(TEXT("name"), Custom->Inputs[i].InputName.ToString());
+		InputNamesArray.Add(MakeShared<FJsonValueObject>(InputObj));
+	}
+	Result->SetArrayField(TEXT("all_inputs"), InputNamesArray);
+	
+	// Include any compilation errors
+	AddCompilationErrorsToResult(Result, Material);
+	
+	return FECACommandResult::Success(Result);
+}
+
+//------------------------------------------------------------------------------
+// FECACommand_SetMaterialProperty
+//------------------------------------------------------------------------------
+
+// Helper to convert JSON value to string for property setting
+// Handles JSON objects by converting them to UE4's ImportText format
+static FString JsonValueToPropertyString(const TSharedPtr<FJsonObject>& Params, const FString& Key)
+{
+	if (!Params->HasField(Key))
+	{
+		return TEXT("");
+	}
+	
+	const TSharedPtr<FJsonValue> Value = Params->TryGetField(Key);
+	
+	switch (Value->Type)
+	{
+		case EJson::Boolean:
+			return Value->AsBool() ? TEXT("True") : TEXT("False");
+		case EJson::Number:
+			return FString::SanitizeFloat(Value->AsNumber());
+		case EJson::String:
+			return Value->AsString();
+		case EJson::Object:
+		{
+			// Convert JSON object to ImportText format for structs
+			const TSharedPtr<FJsonObject>& Obj = Value->AsObject();
+			
+			// Check for color format {r, g, b, a}
+			if (Obj->HasField(TEXT("r")))
+			{
+				double R = Obj->GetNumberField(TEXT("r"));
+				double G = Obj->HasField(TEXT("g")) ? Obj->GetNumberField(TEXT("g")) : 0.0;
+				double B = Obj->HasField(TEXT("b")) ? Obj->GetNumberField(TEXT("b")) : 0.0;
+				double A = Obj->HasField(TEXT("a")) ? Obj->GetNumberField(TEXT("a")) : 1.0;
+				return FString::Printf(TEXT("(R=%f,G=%f,B=%f,A=%f)"), R, G, B, A);
+			}
+			// Check for vector format {x, y, z}
+			else if (Obj->HasField(TEXT("x")))
+			{
+				double X = Obj->GetNumberField(TEXT("x"));
+				double Y = Obj->HasField(TEXT("y")) ? Obj->GetNumberField(TEXT("y")) : 0.0;
+				double Z = Obj->HasField(TEXT("z")) ? Obj->GetNumberField(TEXT("z")) : 0.0;
+				double W = Obj->HasField(TEXT("w")) ? Obj->GetNumberField(TEXT("w")) : 0.0;
+				
+				if (Obj->HasField(TEXT("w")))
+				{
+					return FString::Printf(TEXT("(X=%f,Y=%f,Z=%f,W=%f)"), X, Y, Z, W);
+				}
+				else if (Obj->HasField(TEXT("z")))
+				{
+					return FString::Printf(TEXT("(X=%f,Y=%f,Z=%f)"), X, Y, Z);
+				}
+				else
+				{
+					return FString::Printf(TEXT("(X=%f,Y=%f)"), X, Y);
+				}
+			}
+			// Check for rotator format {pitch, yaw, roll}
+			else if (Obj->HasField(TEXT("pitch")))
+			{
+				double Pitch = Obj->GetNumberField(TEXT("pitch"));
+				double Yaw = Obj->HasField(TEXT("yaw")) ? Obj->GetNumberField(TEXT("yaw")) : 0.0;
+				double Roll = Obj->HasField(TEXT("roll")) ? Obj->GetNumberField(TEXT("roll")) : 0.0;
+				return FString::Printf(TEXT("(Pitch=%f,Yaw=%f,Roll=%f)"), Pitch, Yaw, Roll);
+			}
+			// Check for transform format {location, rotation, scale}
+			else if (Obj->HasField(TEXT("location")))
+			{
+				FString Result = TEXT("(");
+				
+				if (const TSharedPtr<FJsonObject>* LocObj = nullptr; Obj->TryGetObjectField(TEXT("location"), LocObj))
+				{
+					double X = (*LocObj)->HasField(TEXT("x")) ? (*LocObj)->GetNumberField(TEXT("x")) : 0.0;
+					double Y = (*LocObj)->HasField(TEXT("y")) ? (*LocObj)->GetNumberField(TEXT("y")) : 0.0;
+					double Z = (*LocObj)->HasField(TEXT("z")) ? (*LocObj)->GetNumberField(TEXT("z")) : 0.0;
+					Result += FString::Printf(TEXT("Translation=(X=%f,Y=%f,Z=%f)"), X, Y, Z);
+				}
+				
+				if (const TSharedPtr<FJsonObject>* RotObj = nullptr; Obj->TryGetObjectField(TEXT("rotation"), RotObj))
+				{
+					double Pitch = (*RotObj)->HasField(TEXT("pitch")) ? (*RotObj)->GetNumberField(TEXT("pitch")) : 0.0;
+					double Yaw = (*RotObj)->HasField(TEXT("yaw")) ? (*RotObj)->GetNumberField(TEXT("yaw")) : 0.0;
+					double Roll = (*RotObj)->HasField(TEXT("roll")) ? (*RotObj)->GetNumberField(TEXT("roll")) : 0.0;
+					if (Result.Len() > 1) Result += TEXT(",");
+					Result += FString::Printf(TEXT("Rotation=(Pitch=%f,Yaw=%f,Roll=%f)"), Pitch, Yaw, Roll);
+				}
+				
+				if (const TSharedPtr<FJsonObject>* ScaleObj = nullptr; Obj->TryGetObjectField(TEXT("scale"), ScaleObj))
+				{
+					double X = (*ScaleObj)->HasField(TEXT("x")) ? (*ScaleObj)->GetNumberField(TEXT("x")) : 1.0;
+					double Y = (*ScaleObj)->HasField(TEXT("y")) ? (*ScaleObj)->GetNumberField(TEXT("y")) : 1.0;
+					double Z = (*ScaleObj)->HasField(TEXT("z")) ? (*ScaleObj)->GetNumberField(TEXT("z")) : 1.0;
+					if (Result.Len() > 1) Result += TEXT(",");
+					Result += FString::Printf(TEXT("Scale3D=(X=%f,Y=%f,Z=%f)"), X, Y, Z);
+				}
+				
+				Result += TEXT(")");
+				return Result;
+			}
+			
+			// Fallback: serialize as generic (Key=Value,...) format
+			FString Result = TEXT("(");
+			bool bFirst = true;
+			for (const auto& Pair : Obj->Values)
+			{
+				if (!bFirst) Result += TEXT(",");
+				bFirst = false;
+				
+				if (Pair.Value->Type == EJson::Number)
+				{
+					Result += FString::Printf(TEXT("%s=%f"), *Pair.Key, Pair.Value->AsNumber());
+				}
+				else if (Pair.Value->Type == EJson::String)
+				{
+					Result += FString::Printf(TEXT("%s=%s"), *Pair.Key, *Pair.Value->AsString());
+				}
+				else if (Pair.Value->Type == EJson::Boolean)
+				{
+					Result += FString::Printf(TEXT("%s=%s"), *Pair.Key, Pair.Value->AsBool() ? TEXT("True") : TEXT("False"));
+				}
+			}
+			Result += TEXT(")");
+			return Result;
+		}
+		default:
+			return TEXT("");
+	}
+}
+
+// Helper to convert FProperty value to JSON
+static TSharedPtr<FJsonValue> PropertyToJsonValue(FProperty* Property, const void* ValuePtr)
+{
+	if (FBoolProperty* BoolProp = CastField<FBoolProperty>(Property))
+	{
+		return MakeShared<FJsonValueBoolean>(BoolProp->GetPropertyValue(ValuePtr));
+	}
+	else if (FNumericProperty* NumProp = CastField<FNumericProperty>(Property))
+	{
+		if (NumProp->IsFloatingPoint())
+		{
+			return MakeShared<FJsonValueNumber>(NumProp->GetFloatingPointPropertyValue(ValuePtr));
+		}
+		else if (NumProp->IsInteger())
+		{
+			return MakeShared<FJsonValueNumber>(static_cast<double>(NumProp->GetSignedIntPropertyValue(ValuePtr)));
+		}
+	}
+	else if (FEnumProperty* EnumProp = CastField<FEnumProperty>(Property))
+	{
+		UEnum* Enum = EnumProp->GetEnum();
+		FNumericProperty* UnderlyingProp = EnumProp->GetUnderlyingProperty();
+		int64 EnumValue = UnderlyingProp->GetSignedIntPropertyValue(ValuePtr);
+		FString EnumName = Enum->GetNameStringByValue(EnumValue);
+		return MakeShared<FJsonValueString>(EnumName);
+	}
+	else if (FByteProperty* ByteProp = CastField<FByteProperty>(Property))
+	{
+		if (UEnum* Enum = ByteProp->GetIntPropertyEnum())
+		{
+			uint8 ByteValue = *static_cast<const uint8*>(ValuePtr);
+			FString EnumName = Enum->GetNameStringByValue(ByteValue);
+			return MakeShared<FJsonValueString>(EnumName);
+		}
+		else
+		{
+			return MakeShared<FJsonValueNumber>(static_cast<double>(*static_cast<const uint8*>(ValuePtr)));
+		}
+	}
+	else if (FStrProperty* StrProp = CastField<FStrProperty>(Property))
+	{
+		return MakeShared<FJsonValueString>(StrProp->GetPropertyValue(ValuePtr));
+	}
+	else if (FNameProperty* NameProp = CastField<FNameProperty>(Property))
+	{
+		return MakeShared<FJsonValueString>(NameProp->GetPropertyValue(ValuePtr).ToString());
+	}
+	else if (FStructProperty* StructProp = CastField<FStructProperty>(Property))
+	{
+		UScriptStruct* Struct = StructProp->Struct;
+		
+		// Handle common struct types with proper JSON representation
+		if (Struct == TBaseStructure<FLinearColor>::Get())
+		{
+			const FLinearColor* Color = static_cast<const FLinearColor*>(ValuePtr);
+			TSharedPtr<FJsonObject> ColorObj = MakeShared<FJsonObject>();
+			ColorObj->SetNumberField(TEXT("r"), Color->R);
+			ColorObj->SetNumberField(TEXT("g"), Color->G);
+			ColorObj->SetNumberField(TEXT("b"), Color->B);
+			ColorObj->SetNumberField(TEXT("a"), Color->A);
+			return MakeShared<FJsonValueObject>(ColorObj);
+		}
+		else if (Struct == TBaseStructure<FColor>::Get())
+		{
+			const FColor* Color = static_cast<const FColor*>(ValuePtr);
+			TSharedPtr<FJsonObject> ColorObj = MakeShared<FJsonObject>();
+			ColorObj->SetNumberField(TEXT("r"), Color->R);
+			ColorObj->SetNumberField(TEXT("g"), Color->G);
+			ColorObj->SetNumberField(TEXT("b"), Color->B);
+			ColorObj->SetNumberField(TEXT("a"), Color->A);
+			return MakeShared<FJsonValueObject>(ColorObj);
+		}
+		else if (Struct == TBaseStructure<FVector>::Get())
+		{
+			const FVector* Vec = static_cast<const FVector*>(ValuePtr);
+			TSharedPtr<FJsonObject> VecObj = MakeShared<FJsonObject>();
+			VecObj->SetNumberField(TEXT("x"), Vec->X);
+			VecObj->SetNumberField(TEXT("y"), Vec->Y);
+			VecObj->SetNumberField(TEXT("z"), Vec->Z);
+			return MakeShared<FJsonValueObject>(VecObj);
+		}
+		else if (Struct == TBaseStructure<FVector2D>::Get())
+		{
+			const FVector2D* Vec = static_cast<const FVector2D*>(ValuePtr);
+			TSharedPtr<FJsonObject> VecObj = MakeShared<FJsonObject>();
+			VecObj->SetNumberField(TEXT("x"), Vec->X);
+			VecObj->SetNumberField(TEXT("y"), Vec->Y);
+			return MakeShared<FJsonValueObject>(VecObj);
+		}
+		else if (Struct == TBaseStructure<FVector4>::Get())
+		{
+			const FVector4* Vec = static_cast<const FVector4*>(ValuePtr);
+			TSharedPtr<FJsonObject> VecObj = MakeShared<FJsonObject>();
+			VecObj->SetNumberField(TEXT("x"), Vec->X);
+			VecObj->SetNumberField(TEXT("y"), Vec->Y);
+			VecObj->SetNumberField(TEXT("z"), Vec->Z);
+			VecObj->SetNumberField(TEXT("w"), Vec->W);
+			return MakeShared<FJsonValueObject>(VecObj);
+		}
+		else if (Struct == TBaseStructure<FRotator>::Get())
+		{
+			const FRotator* Rot = static_cast<const FRotator*>(ValuePtr);
+			TSharedPtr<FJsonObject> RotObj = MakeShared<FJsonObject>();
+			RotObj->SetNumberField(TEXT("pitch"), Rot->Pitch);
+			RotObj->SetNumberField(TEXT("yaw"), Rot->Yaw);
+			RotObj->SetNumberField(TEXT("roll"), Rot->Roll);
+			return MakeShared<FJsonValueObject>(RotObj);
+		}
+		else if (Struct == TBaseStructure<FTransform>::Get())
+		{
+			const FTransform* Trans = static_cast<const FTransform*>(ValuePtr);
+			TSharedPtr<FJsonObject> TransObj = MakeShared<FJsonObject>();
+			
+			TSharedPtr<FJsonObject> LocObj = MakeShared<FJsonObject>();
+			LocObj->SetNumberField(TEXT("x"), Trans->GetLocation().X);
+			LocObj->SetNumberField(TEXT("y"), Trans->GetLocation().Y);
+			LocObj->SetNumberField(TEXT("z"), Trans->GetLocation().Z);
+			TransObj->SetObjectField(TEXT("location"), LocObj);
+			
+			TSharedPtr<FJsonObject> RotObj = MakeShared<FJsonObject>();
+			FRotator Rot = Trans->GetRotation().Rotator();
+			RotObj->SetNumberField(TEXT("pitch"), Rot.Pitch);
+			RotObj->SetNumberField(TEXT("yaw"), Rot.Yaw);
+			RotObj->SetNumberField(TEXT("roll"), Rot.Roll);
+			TransObj->SetObjectField(TEXT("rotation"), RotObj);
+			
+			TSharedPtr<FJsonObject> ScaleObj = MakeShared<FJsonObject>();
+			ScaleObj->SetNumberField(TEXT("x"), Trans->GetScale3D().X);
+			ScaleObj->SetNumberField(TEXT("y"), Trans->GetScale3D().Y);
+			ScaleObj->SetNumberField(TEXT("z"), Trans->GetScale3D().Z);
+			TransObj->SetObjectField(TEXT("scale"), ScaleObj);
+			
+			return MakeShared<FJsonValueObject>(TransObj);
+		}
+		
+		// For other structs, fall through to string export
+	}
+	
+	// Fallback: export to string
+	FString StringValue;
+	Property->ExportTextItem_Direct(StringValue, ValuePtr, nullptr, nullptr, PPF_None);
+	return MakeShared<FJsonValueString>(StringValue);
+}
+
+// Find property by name, supporting common aliases (snake_case, etc.)
+static FProperty* FindMaterialProperty(UClass* Class, const FString& PropertyName)
+{
+	// First try direct match
+	if (FProperty* Prop = Class->FindPropertyByName(FName(*PropertyName)))
+	{
+		return Prop;
+	}
+	
+	// Try with 'b' prefix for bools
+	FString BoolName = TEXT("b") + PropertyName;
+	if (FProperty* Prop = Class->FindPropertyByName(FName(*BoolName)))
+	{
+		return Prop;
+	}
+	
+	// Convert snake_case to PascalCase and try again
+	FString PascalCase;
+	bool bCapitalizeNext = true;
+	for (TCHAR Char : PropertyName)
+	{
+		if (Char == '_')
+		{
+			bCapitalizeNext = true;
+		}
+		else
+		{
+			PascalCase += bCapitalizeNext ? FChar::ToUpper(Char) : Char;
+			bCapitalizeNext = false;
+		}
+	}
+	
+	if (FProperty* Prop = Class->FindPropertyByName(FName(*PascalCase)))
+	{
+		return Prop;
+	}
+	
+	// Try with 'b' prefix for PascalCase
+	BoolName = TEXT("b") + PascalCase;
+	if (FProperty* Prop = Class->FindPropertyByName(FName(*BoolName)))
+	{
+		return Prop;
+	}
+	
+	return nullptr;
+}
+
+FECACommandResult FECACommand_SetMaterialProperty::Execute(const TSharedPtr<FJsonObject>& Params)
+{
+	FString MaterialPath = Params->GetStringField(TEXT("material_path"));
+	FString PropertyName = Params->GetStringField(TEXT("property"));
+	
+	// Load the material
+	UMaterial* Material = LoadObject<UMaterial>(nullptr, *MaterialPath);
+	if (!Material)
+	{
+		return FECACommandResult::Error(FString::Printf(TEXT("Failed to load material: %s"), *MaterialPath));
+	}
+	
+	// Find the property using reflection
+	FProperty* Property = FindMaterialProperty(UMaterial::StaticClass(), PropertyName);
+	if (!Property)
+	{
+		// Build list of editable properties for error message
+		TArray<FString> ValidProperties;
+		for (TFieldIterator<FProperty> PropIt(UMaterial::StaticClass()); PropIt; ++PropIt)
+		{
+			if (PropIt->HasAnyPropertyFlags(CPF_Edit))
+			{
+				ValidProperties.Add(PropIt->GetName());
+			}
+		}
+		return FECACommandResult::Error(FString::Printf(TEXT("Property '%s' not found on UMaterial. Some valid properties: %s"), 
+			*PropertyName, *FString::Join(ValidProperties, TEXT(", ")).Left(500)));
+	}
+	
+	// Get the value as a string for ImportText
+	FString ValueStr = JsonValueToPropertyString(Params, TEXT("value"));
+	
+	// Get pointer to property value
+	void* ValuePtr = Property->ContainerPtrToValuePtr<void>(Material);
+	
+	// Store old value for reporting
+	TSharedPtr<FJsonValue> OldValue = PropertyToJsonValue(Property, ValuePtr);
+	
+	// Use ImportText to set the value (handles enums, bools, numbers, etc.)
+	const TCHAR* ImportResult = Property->ImportText_Direct(*ValueStr, ValuePtr, Material, PPF_None);
+	if (!ImportResult)
+	{
+		return FECACommandResult::Error(FString::Printf(TEXT("Failed to set property '%s' to value '%s'"), *PropertyName, *ValueStr));
+	}
+	
+	// Get new value for reporting
+	TSharedPtr<FJsonValue> NewValue = PropertyToJsonValue(Property, ValuePtr);
+	
+	// Mark the material as needing recompilation
+	ClearMaterialCompilationErrors(Material);
+	Material->PostEditChange();
+	Material->MarkPackageDirty();
+	
+	// Refresh the material editor UI if open
+	RefreshMaterialEditorUI(Material);
+	
+	TSharedPtr<FJsonObject> Result = MakeResult();
+	Result->SetStringField(TEXT("material_path"), MaterialPath);
+	Result->SetStringField(TEXT("property"), Property->GetName());
+	Result->SetField(TEXT("old_value"), OldValue);
+	Result->SetField(TEXT("new_value"), NewValue);
+	Result->SetBoolField(TEXT("success"), true);
+	
+	// Include any compilation errors
+	AddCompilationErrorsToResult(Result, Material);
+	
+	return FECACommandResult::Success(Result);
+}
+
+//------------------------------------------------------------------------------
+// FECACommand_GetMaterialProperty
+//------------------------------------------------------------------------------
+
+FECACommandResult FECACommand_GetMaterialProperty::Execute(const TSharedPtr<FJsonObject>& Params)
+{
+	FString MaterialPath = Params->GetStringField(TEXT("material_path"));
+	FString PropertyName = Params->HasField(TEXT("property")) ? Params->GetStringField(TEXT("property")) : TEXT("");
+	
+	// Load the material
+	UMaterial* Material = LoadObject<UMaterial>(nullptr, *MaterialPath);
+	if (!Material)
+	{
+		return FECACommandResult::Error(FString::Printf(TEXT("Failed to load material: %s"), *MaterialPath));
+	}
+	
+	TSharedPtr<FJsonObject> Result = MakeResult();
+	Result->SetStringField(TEXT("material_path"), MaterialPath);
+	
+	if (!PropertyName.IsEmpty())
+	{
+		// Get specific property
+		FProperty* Property = FindMaterialProperty(UMaterial::StaticClass(), PropertyName);
+		if (!Property)
+		{
+			return FECACommandResult::Error(FString::Printf(TEXT("Property '%s' not found on UMaterial"), *PropertyName));
+		}
+		
+		void* ValuePtr = Property->ContainerPtrToValuePtr<void>(Material);
+		Result->SetField(Property->GetName(), PropertyToJsonValue(Property, ValuePtr));
+	}
+	else
+	{
+		// Return all editable properties
+		TSharedPtr<FJsonObject> PropertiesObj = MakeShared<FJsonObject>();
+		
+		for (TFieldIterator<FProperty> PropIt(UMaterial::StaticClass()); PropIt; ++PropIt)
+		{
+			FProperty* Property = *PropIt;
+			
+			// Only include EditAnywhere properties (the ones visible in the editor)
+			if (!Property->HasAnyPropertyFlags(CPF_Edit))
+			{
+				continue;
+			}
+			
+			// Skip deprecated and transient properties
+			if (Property->HasAnyPropertyFlags(CPF_Deprecated | CPF_Transient))
+			{
+				continue;
+			}
+			
+			void* ValuePtr = Property->ContainerPtrToValuePtr<void>(Material);
+			PropertiesObj->SetField(Property->GetName(), PropertyToJsonValue(Property, ValuePtr));
+		}
+		
+		Result->SetObjectField(TEXT("properties"), PropertiesObj);
+	}
+	
+	return FECACommandResult::Success(Result);
+}
+
+
+//------------------------------------------------------------------------------
+// AutoLayoutMaterialGraph - Automatically arrange material expression nodes
+//------------------------------------------------------------------------------
+
+FECACommandResult FECACommand_AutoLayoutMaterialGraph::Execute(const TSharedPtr<FJsonObject>& Params)
+{
+	FString MaterialPath = Params->GetStringField(TEXT("material_path"));
+	
+	UMaterial* Material = LoadMaterialByPath(MaterialPath);
+	if (!Material || !Material->GetEditorOnlyData())
+	{
+		return FECACommandResult::Error(FString::Printf(TEXT("Failed to load material: %s"), *MaterialPath));
+	}
+	
+	// Get layout parameters
+	FString Strategy = TEXT("horizontal");
+	if (Params->HasField(TEXT("strategy")))
+	{
+		Strategy = Params->GetStringField(TEXT("strategy")).ToLower();
+	}
+	
+	int32 SpacingX = 300;
+	int32 SpacingY = 150;
+	if (Params->HasField(TEXT("spacing_x")))
+	{
+		SpacingX = static_cast<int32>(Params->GetNumberField(TEXT("spacing_x")));
+	}
+	if (Params->HasField(TEXT("spacing_y")))
+	{
+		SpacingY = static_cast<int32>(Params->GetNumberField(TEXT("spacing_y")));
+	}
+	
+	// Get specific node IDs if provided
+	TSet<FGuid> SpecificNodeIds;
+	const TArray<TSharedPtr<FJsonValue>>* NodeIdsArray;
+	if (Params->TryGetArrayField(TEXT("node_ids"), NodeIdsArray) && NodeIdsArray)
+	{
+		for (const TSharedPtr<FJsonValue>& IdValue : *NodeIdsArray)
+		{
+			FString IdStr;
+			if (IdValue->TryGetString(IdStr))
+			{
+				FGuid Guid;
+				if (FGuid::Parse(IdStr, Guid))
+				{
+					SpecificNodeIds.Add(Guid);
+				}
+			}
+		}
+	}
+	
+	// Collect nodes to layout
+	TArray<UMaterialExpression*> NodesToLayout;
+	for (UMaterialExpression* Expr : Material->GetEditorOnlyData()->ExpressionCollection.Expressions)
+	{
+		if (!Expr)
+		{
+			continue;
+		}
+		
+		// Filter by specific node IDs if provided
+		if (SpecificNodeIds.Num() > 0)
+		{
+			if (!SpecificNodeIds.Contains(Expr->MaterialExpressionGuid))
+			{
+				continue;
+			}
+		}
+		
+		NodesToLayout.Add(Expr);
+	}
+	
+	if (NodesToLayout.Num() == 0)
+	{
+		return FECACommandResult::Error(TEXT("No nodes to layout"));
+	}
+	
+	// Build connection map: which expressions connect to which
+	// We'll trace from material inputs back through the graph
+	TMap<UMaterialExpression*, TSet<UMaterialExpression*>> ConsumersMap; // Who consumes this expression
+	TMap<UMaterialExpression*, TSet<UMaterialExpression*>> ProducersMap; // Who produces input for this expression
+	
+	// Helper to get expressions that feed into another expression
+	auto GetInputExpressions = [](UMaterialExpression* Expr) -> TArray<UMaterialExpression*>
+	{
+		TArray<UMaterialExpression*> Inputs;
+		for (int32 i = 0; ; i++)
+		{
+			FExpressionInput* Input = Expr->GetInput(i);
+			if (!Input)
+			{
+				break;
+			}
+			if (Input->Expression)
+			{
+				Inputs.AddUnique(Input->Expression);
+			}
+		}
+		return Inputs;
+	};
+	
+	// Build the connection maps
+	for (UMaterialExpression* Expr : NodesToLayout)
+	{
+		TArray<UMaterialExpression*> Inputs = GetInputExpressions(Expr);
+		for (UMaterialExpression* InputExpr : Inputs)
+		{
+			if (NodesToLayout.Contains(InputExpr))
+			{
+				ConsumersMap.FindOrAdd(InputExpr).Add(Expr);
+				ProducersMap.FindOrAdd(Expr).Add(InputExpr);
+			}
+		}
+	}
+	
+	// Find material input connections (these are our "roots" in reverse)
+	// Material inputs: BaseColor, Metallic, Specular, Roughness, EmissiveColor, Opacity, 
+	// OpacityMask, Normal, Tangent, WorldPositionOffset, etc.
+	TArray<UMaterialExpression*> MaterialInputs;
+	
+	// Check each material input property
+	auto CheckMaterialInput = [&](FExpressionInput& Input)
+	{
+		if (Input.Expression && NodesToLayout.Contains(Input.Expression))
+		{
+			MaterialInputs.AddUnique(Input.Expression);
+		}
+	};
+	
+	// Get all material input connections
+	if (Material->GetEditorOnlyData())
+	{
+		CheckMaterialInput(Material->GetEditorOnlyData()->BaseColor);
+		CheckMaterialInput(Material->GetEditorOnlyData()->Metallic);
+		CheckMaterialInput(Material->GetEditorOnlyData()->Specular);
+		CheckMaterialInput(Material->GetEditorOnlyData()->Roughness);
+		CheckMaterialInput(Material->GetEditorOnlyData()->EmissiveColor);
+		CheckMaterialInput(Material->GetEditorOnlyData()->Opacity);
+		CheckMaterialInput(Material->GetEditorOnlyData()->OpacityMask);
+		CheckMaterialInput(Material->GetEditorOnlyData()->Normal);
+		CheckMaterialInput(Material->GetEditorOnlyData()->Tangent);
+		CheckMaterialInput(Material->GetEditorOnlyData()->WorldPositionOffset);
+		CheckMaterialInput(Material->GetEditorOnlyData()->SubsurfaceColor);
+		CheckMaterialInput(Material->GetEditorOnlyData()->AmbientOcclusion);
+		CheckMaterialInput(Material->GetEditorOnlyData()->Refraction);
+		CheckMaterialInput(Material->GetEditorOnlyData()->PixelDepthOffset);
+		CheckMaterialInput(Material->GetEditorOnlyData()->ShadingModelFromMaterialExpression);
+		CheckMaterialInput(Material->GetEditorOnlyData()->Displacement);
+		CheckMaterialInput(Material->GetEditorOnlyData()->FrontMaterial);
+		CheckMaterialInput(Material->GetEditorOnlyData()->SurfaceThickness);
+		CheckMaterialInput(Material->GetEditorOnlyData()->Anisotropy);
+		CheckMaterialInput(Material->GetEditorOnlyData()->ClearCoat);
+		CheckMaterialInput(Material->GetEditorOnlyData()->ClearCoatRoughness);
+	}
+	
+	// Assign depths by traversing backward from material inputs
+	TMap<UMaterialExpression*, int32> NodeDepths;
+	
+	// BFS from material inputs going backward
+	TQueue<UMaterialExpression*> TraversalQueue;
+	TSet<UMaterialExpression*> Visited;
+	
+	// Start with nodes directly connected to material inputs (depth 0 = closest to material)
+	for (UMaterialExpression* Input : MaterialInputs)
+	{
+		NodeDepths.Add(Input, 0);
+		TraversalQueue.Enqueue(Input);
+		Visited.Add(Input);
+	}
+	
+	// If no material inputs, use nodes with consumers (outputs go somewhere)
+	if (MaterialInputs.Num() == 0)
+	{
+		for (UMaterialExpression* Expr : NodesToLayout)
+		{
+			if (ConsumersMap.Contains(Expr) && ConsumersMap[Expr].Num() > 0)
+			{
+				// This node feeds into something but isn't a root yet
+				if (!ProducersMap.Contains(Expr) || ProducersMap[Expr].Num() == 0)
+				{
+					// This has no producers, so it could be a parameter or constant
+					NodeDepths.Add(Expr, 0);
+					TraversalQueue.Enqueue(Expr);
+					Visited.Add(Expr);
+				}
+			}
+		}
+	}
+	
+	// If still no roots, just start with the first node
+	if (Visited.Num() == 0 && NodesToLayout.Num() > 0)
+	{
+		NodeDepths.Add(NodesToLayout[0], 0);
+		TraversalQueue.Enqueue(NodesToLayout[0]);
+		Visited.Add(NodesToLayout[0]);
+	}
+	
+	// Traverse backward through producers
+	while (!TraversalQueue.IsEmpty())
+	{
+		UMaterialExpression* Current;
+		TraversalQueue.Dequeue(Current);
+		
+		int32 CurrentDepth = NodeDepths.FindRef(Current);
+		
+		// Get all inputs (producers) for this node
+		TArray<UMaterialExpression*> Inputs = GetInputExpressions(Current);
+		for (UMaterialExpression* InputExpr : Inputs)
+		{
+			if (NodesToLayout.Contains(InputExpr) && !Visited.Contains(InputExpr))
+			{
+				Visited.Add(InputExpr);
+				NodeDepths.Add(InputExpr, CurrentDepth + 1);
+				TraversalQueue.Enqueue(InputExpr);
+			}
+		}
+	}
+	
+	// Handle any unvisited nodes (disconnected nodes)
+	for (UMaterialExpression* Expr : NodesToLayout)
+	{
+		if (!NodeDepths.Contains(Expr))
+		{
+			// Place disconnected nodes at the end
+			NodeDepths.Add(Expr, 0);
+		}
+	}
+	
+	// Group nodes by depth
+	TMap<int32, TArray<UMaterialExpression*>> NodesByDepth;
+	int32 MaxDepth = 0;
+	
+	for (const auto& Pair : NodeDepths)
+	{
+		NodesByDepth.FindOrAdd(Pair.Value).Add(Pair.Key);
+		MaxDepth = FMath::Max(MaxDepth, Pair.Value);
+	}
+	
+	// Calculate starting position (material result node position, or use a default)
+	// Material result is typically on the right side, so we'll position left of it
+	int32 StartX = 0;
+	int32 StartY = 0;
+	
+	// Use the minimum position of existing nodes as anchor
+	for (UMaterialExpression* Expr : NodesToLayout)
+	{
+		if (StartX == 0 && StartY == 0)
+		{
+			StartX = Expr->MaterialExpressionEditorX;
+			StartY = Expr->MaterialExpressionEditorY;
+		}
+		else
+		{
+			StartX = FMath::Min(StartX, Expr->MaterialExpressionEditorX);
+			StartY = FMath::Min(StartY, Expr->MaterialExpressionEditorY);
+		}
+	}
+	
+	// Position nodes based on strategy
+	// Note: In materials, depth 0 = closest to material output (right side)
+	// Higher depths = further left (input nodes, parameters, etc.)
+	int32 NodesPositioned = 0;
+	
+	if (Strategy == TEXT("vertical"))
+	{
+		// Vertical: arrange top-to-bottom, with material inputs at top
+		for (int32 Depth = 0; Depth <= MaxDepth; Depth++)
+		{
+			TArray<UMaterialExpression*>* NodesAtDepth = NodesByDepth.Find(Depth);
+			if (!NodesAtDepth)
+			{
+				continue;
+			}
+			
+			int32 Y = StartY + Depth * SpacingY;
+			int32 X = StartX;
+			
+			for (UMaterialExpression* Expr : *NodesAtDepth)
+			{
+				Expr->MaterialExpressionEditorX = X;
+				Expr->MaterialExpressionEditorY = Y;
+				X += SpacingX;
+				NodesPositioned++;
+			}
+		}
+	}
+	else if (Strategy == TEXT("tree"))
+	{
+		// Tree: center producers under their consumers
+		TMap<UMaterialExpression*, int32> SubtreeWidths;
+		
+		// Calculate subtree widths from deepest to shallowest
+		for (int32 Depth = MaxDepth; Depth >= 0; Depth--)
+		{
+			TArray<UMaterialExpression*>* NodesAtDepth = NodesByDepth.Find(Depth);
+			if (!NodesAtDepth)
+			{
+				continue;
+			}
+			
+			for (UMaterialExpression* Expr : *NodesAtDepth)
+			{
+				int32 Width = 1;
+				
+				// Sum widths of producer nodes
+				TArray<UMaterialExpression*> Inputs = GetInputExpressions(Expr);
+				for (UMaterialExpression* InputExpr : Inputs)
+				{
+					if (SubtreeWidths.Contains(InputExpr))
+					{
+						Width += SubtreeWidths[InputExpr];
+					}
+				}
+				
+				SubtreeWidths.Add(Expr, FMath::Max(1, Width));
+			}
+		}
+		
+		// Position nodes
+		TMap<UMaterialExpression*, int32> NodeYPositions;
+		int32 CurrentY = StartY;
+		
+		// Position depth 0 (material input nodes) first
+		TArray<UMaterialExpression*>* RootNodes = NodesByDepth.Find(0);
+		if (RootNodes)
+		{
+			for (UMaterialExpression* Root : *RootNodes)
+			{
+				int32 Width = SubtreeWidths.FindRef(Root);
+				NodeYPositions.Add(Root, CurrentY + (Width * SpacingY) / 2);
+				Root->MaterialExpressionEditorX = StartX;
+				Root->MaterialExpressionEditorY = NodeYPositions[Root];
+				CurrentY += Width * SpacingY;
+				NodesPositioned++;
+			}
+		}
+		
+		// Position remaining depths
+		for (int32 Depth = 1; Depth <= MaxDepth; Depth++)
+		{
+			TArray<UMaterialExpression*>* NodesAtDepth = NodesByDepth.Find(Depth);
+			if (!NodesAtDepth)
+			{
+				continue;
+			}
+			
+			// Position to the left of previous depth
+			int32 X = StartX - Depth * SpacingX;
+			
+			for (UMaterialExpression* Expr : *NodesAtDepth)
+			{
+				if (!NodeYPositions.Contains(Expr))
+				{
+					Expr->MaterialExpressionEditorX = X;
+					Expr->MaterialExpressionEditorY = StartY;
+					NodeYPositions.Add(Expr, StartY);
+				}
+				else
+				{
+					Expr->MaterialExpressionEditorX = X;
+					Expr->MaterialExpressionEditorY = NodeYPositions[Expr];
+				}
+				NodesPositioned++;
+			}
+		}
+	}
+	else if (Strategy == TEXT("compact"))
+	{
+		// Compact: grid packing
+		int32 NodesPerRow = FMath::Max(1, FMath::CeilToInt(FMath::Sqrt((float)NodesToLayout.Num())));
+		int32 Index = 0;
+		
+		for (UMaterialExpression* Expr : NodesToLayout)
+		{
+			int32 Row = Index / NodesPerRow;
+			int32 Col = Index % NodesPerRow;
+			
+			Expr->MaterialExpressionEditorX = StartX + Col * SpacingX;
+			Expr->MaterialExpressionEditorY = StartY + Row * SpacingY;
+			
+			Index++;
+			NodesPositioned++;
+		}
+	}
+	else // horizontal (default)
+	{
+		// Horizontal: arrange left-to-right
+		// Depth 0 (closest to material) on the right, higher depths on the left
+		for (int32 Depth = 0; Depth <= MaxDepth; Depth++)
+		{
+			TArray<UMaterialExpression*>* NodesAtDepth = NodesByDepth.Find(Depth);
+			if (!NodesAtDepth)
+			{
+				continue;
+			}
+			
+			// Sort by Y position to maintain relative vertical order
+			NodesAtDepth->Sort([](const UMaterialExpression& A, const UMaterialExpression& B) {
+				return A.MaterialExpressionEditorY < B.MaterialExpressionEditorY;
+			});
+			
+			// Position: depth 0 is rightmost, higher depths move left
+			int32 X = StartX - Depth * SpacingX;
+			int32 Y = StartY;
+			
+			for (UMaterialExpression* Expr : *NodesAtDepth)
+			{
+				Expr->MaterialExpressionEditorX = X;
+				Expr->MaterialExpressionEditorY = Y;
+				Y += SpacingY;
+				NodesPositioned++;
+			}
+		}
+	}
+	
+	// Mark material as dirty and refresh
+	Material->PostEditChange();
+	Material->MarkPackageDirty();
+	RefreshMaterialEditorUI(Material);
+	
+	// Build result
+	TSharedPtr<FJsonObject> Result = MakeResult();
+	Result->SetStringField(TEXT("material_path"), MaterialPath);
+	Result->SetStringField(TEXT("strategy"), Strategy);
+	Result->SetNumberField(TEXT("nodes_positioned"), NodesPositioned);
+	Result->SetNumberField(TEXT("max_depth"), MaxDepth);
+	Result->SetNumberField(TEXT("spacing_x"), SpacingX);
+	Result->SetNumberField(TEXT("spacing_y"), SpacingY);
+	
+	// Include positioned node info
+	TArray<TSharedPtr<FJsonValue>> PositionedNodesArray;
+	for (UMaterialExpression* Expr : NodesToLayout)
+	{
+		TSharedPtr<FJsonObject> NodeInfo = MakeShared<FJsonObject>();
+		NodeInfo->SetStringField(TEXT("node_id"), Expr->MaterialExpressionGuid.ToString());
+		NodeInfo->SetStringField(TEXT("node_type"), Expr->GetClass()->GetName().Replace(TEXT("MaterialExpression"), TEXT("")));
+		NodeInfo->SetNumberField(TEXT("x"), Expr->MaterialExpressionEditorX);
+		NodeInfo->SetNumberField(TEXT("y"), Expr->MaterialExpressionEditorY);
+		NodeInfo->SetNumberField(TEXT("depth"), NodeDepths.FindRef(Expr));
+		PositionedNodesArray.Add(MakeShared<FJsonValueObject>(NodeInfo));
+	}
+	Result->SetArrayField(TEXT("positioned_nodes"), PositionedNodesArray);
+	
+	return FECACommandResult::Success(Result);
+}
+
+//------------------------------------------------------------------------------
+// RenameParameterGroup - Rename parameter groups in a material
+//------------------------------------------------------------------------------
+
+FECACommandResult FECACommand_RenameParameterGroup::Execute(const TSharedPtr<FJsonObject>& Params)
+{
+	FString MaterialPath = Params->GetStringField(TEXT("material_path"));
+	FString OldGroupName;
+	FString NewGroupName;
+	
+	if (!GetStringParam(Params, TEXT("old_group_name"), OldGroupName))
+	{
+		return FECACommandResult::Error(TEXT("Missing required parameter: old_group_name"));
+	}
+	
+	if (!GetStringParam(Params, TEXT("new_group_name"), NewGroupName))
+	{
+		return FECACommandResult::Error(TEXT("Missing required parameter: new_group_name"));
+	}
+	
+	UMaterial* Material = LoadMaterialByPath(MaterialPath);
+	if (!Material || !Material->GetEditorOnlyData())
+	{
+		return FECACommandResult::Error(FString::Printf(TEXT("Material not found: %s"), *MaterialPath));
+	}
+	
+	// Convert group names to FName for comparison
+	FName OldGroup = OldGroupName.IsEmpty() ? NAME_None : FName(*OldGroupName);
+	FName NewGroup = NewGroupName.IsEmpty() ? NAME_None : FName(*NewGroupName);
+	
+	Material->PreEditChange(nullptr);
+	
+	int32 RenamedCount = 0;
+	TArray<TSharedPtr<FJsonValue>> RenamedParameters;
+	
+	// Iterate through all expressions and update Group property on parameter expressions
+	for (UMaterialExpression* Expr : Material->GetEditorOnlyData()->ExpressionCollection.Expressions)
+	{
+		if (!Expr)
+		{
+			continue;
+		}
+		
+		// Check if this is a parameter expression with the Group property
+		// Use reflection to access the Group property which exists on UMaterialExpressionParameter and derivatives
+		FProperty* GroupProperty = Expr->GetClass()->FindPropertyByName(TEXT("Group"));
+		if (!GroupProperty)
+		{
+			continue;
+		}
+		
+		FNameProperty* GroupNameProperty = CastField<FNameProperty>(GroupProperty);
+		if (!GroupNameProperty)
+		{
+			continue;
+		}
+		
+		// Get the current group value
+		void* GroupValuePtr = GroupProperty->ContainerPtrToValuePtr<void>(Expr);
+		FName CurrentGroup = GroupNameProperty->GetPropertyValue(GroupValuePtr);
+		
+		// Check if this parameter's group matches the old group name
+		if (CurrentGroup == OldGroup)
+		{
+			// Update to the new group name
+			GroupNameProperty->SetPropertyValue(GroupValuePtr, NewGroup);
+			RenamedCount++;
+			
+			// Record the renamed parameter
+			TSharedPtr<FJsonObject> ParamInfo = MakeShared<FJsonObject>();
+			ParamInfo->SetStringField(TEXT("node_id"), Expr->MaterialExpressionGuid.ToString());
+			ParamInfo->SetStringField(TEXT("node_type"), Expr->GetClass()->GetName().Replace(TEXT("MaterialExpression"), TEXT("")));
+			
+			// Try to get the parameter name if available
+			FProperty* ParamNameProperty = Expr->GetClass()->FindPropertyByName(TEXT("ParameterName"));
+			if (FNameProperty* ParamNameProp = CastField<FNameProperty>(ParamNameProperty))
+			{
+				void* ParamNamePtr = ParamNameProperty->ContainerPtrToValuePtr<void>(Expr);
+				FName ParamName = ParamNameProp->GetPropertyValue(ParamNamePtr);
+				ParamInfo->SetStringField(TEXT("parameter_name"), ParamName.ToString());
+			}
+			
+			ParamInfo->SetStringField(TEXT("old_group"), OldGroupName);
+			ParamInfo->SetStringField(TEXT("new_group"), NewGroupName);
+			RenamedParameters.Add(MakeShared<FJsonValueObject>(ParamInfo));
+		}
+	}
+	
+	ClearMaterialCompilationErrors(Material);
+	Material->PostEditChange();
+	Material->MarkPackageDirty();
+	
+	// Refresh the material editor UI if open
+	RefreshMaterialEditorUI(Material);
+	
+	TSharedPtr<FJsonObject> Result = MakeResult();
+	Result->SetStringField(TEXT("material_path"), MaterialPath);
+	Result->SetStringField(TEXT("old_group_name"), OldGroupName);
+	Result->SetStringField(TEXT("new_group_name"), NewGroupName);
+	Result->SetNumberField(TEXT("renamed_count"), RenamedCount);
+	Result->SetArrayField(TEXT("renamed_parameters"), RenamedParameters);
+	
+	// Include any compilation errors
+	AddCompilationErrorsToResult(Result, Material);
+	
+	return FECACommandResult::Success(Result);
+}
+
+//------------------------------------------------------------------------------
+// List Parameter Groups
+//------------------------------------------------------------------------------
+
+FECACommandResult FECACommand_ListParameterGroups::Execute(const TSharedPtr<FJsonObject>& Params)
+{
+	FString MaterialPath = Params->GetStringField(TEXT("material_path"));
+	
+	UMaterial* Material = LoadMaterialByPath(MaterialPath);
+	if (!Material || !Material->GetEditorOnlyData())
+	{
+		return FECACommandResult::Error(FString::Printf(TEXT("Material not found: %s"), *MaterialPath));
+	}
+	
+	// Map group name to list of parameters
+	TMap<FString, TArray<TSharedPtr<FJsonObject>>> GroupsMap;
+	
+	// Iterate through all expressions
+	for (UMaterialExpression* Expr : Material->GetEditorOnlyData()->ExpressionCollection.Expressions)
+	{
+		if (!Expr)
+		{
+			continue;
+		}
+		
+		// Check if this is a parameter expression with the Group property
+		FProperty* GroupProperty = Expr->GetClass()->FindPropertyByName(TEXT("Group"));
+		if (!GroupProperty)
+		{
+			continue;
+		}
+		
+		FNameProperty* GroupNameProperty = CastField<FNameProperty>(GroupProperty);
+		if (!GroupNameProperty)
+		{
+			continue;
+		}
+		
+		// Get the current group value
+		void* GroupValuePtr = GroupProperty->ContainerPtrToValuePtr<void>(Expr);
+		FName CurrentGroup = GroupNameProperty->GetPropertyValue(GroupValuePtr);
+		
+		// Get group name as string (None becomes empty string displayed as "(Default)")
+		FString GroupName = CurrentGroup.IsNone() ? TEXT("") : CurrentGroup.ToString();
+		
+		// Build parameter info
+		TSharedPtr<FJsonObject> ParamInfo = MakeShared<FJsonObject>();
+		ParamInfo->SetStringField(TEXT("node_id"), Expr->MaterialExpressionGuid.ToString());
+		ParamInfo->SetStringField(TEXT("node_type"), Expr->GetClass()->GetName().Replace(TEXT("MaterialExpression"), TEXT("")));
+		
+		// Try to get the parameter name
+		FProperty* ParamNameProperty = Expr->GetClass()->FindPropertyByName(TEXT("ParameterName"));
+		if (FNameProperty* ParamNameProp = CastField<FNameProperty>(ParamNameProperty))
+		{
+			void* ParamNamePtr = ParamNameProperty->ContainerPtrToValuePtr<void>(Expr);
+			FName ParamName = ParamNameProp->GetPropertyValue(ParamNamePtr);
+			ParamInfo->SetStringField(TEXT("parameter_name"), ParamName.ToString());
+		}
+		
+		// Add to the group
+		GroupsMap.FindOrAdd(GroupName).Add(ParamInfo);
+	}
+	
+	// Build result
+	TSharedPtr<FJsonObject> Result = MakeResult();
+	Result->SetStringField(TEXT("material_path"), MaterialPath);
+	
+	TArray<TSharedPtr<FJsonValue>> GroupsArray;
+	for (auto& Pair : GroupsMap)
+	{
+		TSharedPtr<FJsonObject> GroupInfo = MakeShared<FJsonObject>();
+		GroupInfo->SetStringField(TEXT("group_name"), Pair.Key.IsEmpty() ? TEXT("(Default)") : Pair.Key);
+		GroupInfo->SetStringField(TEXT("group_name_raw"), Pair.Key);
+		GroupInfo->SetNumberField(TEXT("parameter_count"), Pair.Value.Num());
+		
+		TArray<TSharedPtr<FJsonValue>> ParamsArray;
+		for (auto& ParamJson : Pair.Value)
+		{
+			ParamsArray.Add(MakeShared<FJsonValueObject>(ParamJson));
+		}
+		GroupInfo->SetArrayField(TEXT("parameters"), ParamsArray);
+		
+		GroupsArray.Add(MakeShared<FJsonValueObject>(GroupInfo));
+	}
+	
+	Result->SetArrayField(TEXT("groups"), GroupsArray);
+	Result->SetNumberField(TEXT("total_groups"), GroupsMap.Num());
+	
+	return FECACommandResult::Success(Result);
+}
