ECABridge experimental MCP plugin

diff --git a/Engine/Plugins/Experimental/ECABridge/Source/ECABridge/Private/Commands/ECAMeshCommands.cpp b/Engine/Plugins/Experimental/ECABridge/Source/ECABridge/Private/Commands/ECAMeshCommands.cpp
new file mode 100644
--- /dev/null
+++ b/Engine/Plugins/Experimental/ECABridge/Source/ECABridge/Private/Commands/ECAMeshCommands.cpp
@@ -0,0 +1,6387 @@
+// Copyright Epic Games, Inc. All Rights Reserved.
+
+#include "Commands/ECAMeshCommands.h"
+#include "Commands/ECACommand.h"
+
+// Core
+#include "Editor.h"
+#include "Engine/StaticMesh.h"
+#include "Engine/Texture2D.h"
+#include "Materials/Material.h"
+#include "Materials/MaterialInterface.h"
+#include "Components/StaticMeshComponent.h"
+#include "Engine/StaticMeshActor.h"
+#include "GameFramework/Actor.h"
+#include "EngineUtils.h"
+#include "HAL/FileManager.h"
+#include "Misc/Paths.h"
+#include "Misc/FileHelper.h"
+#include "UObject/SavePackage.h"
+
+// Mesh Description
+#include "MeshDescription.h"
+#include "MeshDescriptionBuilder.h"
+#include "StaticMeshAttributes.h"
+#include "PhysicsEngine/BodySetup.h"
+
+// GeometryScript / GeometryCore
+#include "UDynamicMesh.h"
+#include "DynamicMesh/DynamicMesh3.h"
+#include "DynamicMesh/MeshNormals.h"
+#include "DynamicMesh/MeshTransforms.h"
+#include "DynamicMeshToMeshDescription.h"
+#include "MeshDescriptionToDynamicMesh.h"
+#include "Operations/MeshPlaneCut.h"
+#include "Operations/MeshMirror.h"
+#include "Operations/HoleFiller.h"
+#include "Operations/SimpleHoleFiller.h"
+#include "DynamicMeshEditor.h"
+#include "Generators/MinimalBoxMeshGenerator.h"
+#include "Generators/SphereGenerator.h"
+#include "Generators/CapsuleGenerator.h"
+#include "Generators/GridBoxMeshGenerator.h"
+#include "Generators/RectangleMeshGenerator.h"
+#include "Generators/DiscMeshGenerator.h"
+#include "Generators/SweepGenerator.h"
+#include "Generators/RevolveGenerator.h"
+#include "Operations/SubdividePoly.h"
+#include "Operations/UniformTessellate.h"
+#include "Operations/PNTriangles.h"
+#include "GroupTopology.h"
+#include "CompGeom/PolygonTriangulation.h"
+#include "Operations/MeshBoolean.h"
+#include "MeshSimplification.h"
+
+// For depth displacement
+#include "Components/SceneCaptureComponent2D.h"
+#include "Engine/SceneCapture2D.h"
+#include "Engine/TextureRenderTarget2D.h"
+#include "TextureResource.h"
+#include "LevelEditor.h"
+#include "IAssetViewport.h"
+#include "EditorViewportClient.h"
+#include "Misc/Base64.h"
+#include "ImageUtils.h"
+#include "IImageWrapper.h"
+#include "IImageWrapperModule.h"
+#include "Remesher.h"
+#include "MeshConstraintsUtil.h"
+
+// ProceduralMeshComponent
+#include "ProceduralMeshComponent.h"
+#include "KismetProceduralMeshLibrary.h"
+
+// Asset tools
+#include "AssetRegistry/AssetRegistryModule.h"
+#include "AssetToolsModule.h"
+#include "IAssetTools.h"
+
+using namespace UE::Geometry;
+
+// Register all mesh commands
+REGISTER_ECA_COMMAND(FECACommand_CreatePrimitive)
+REGISTER_ECA_COMMAND(FECACommand_CreateMesh)
+REGISTER_ECA_COMMAND(FECACommand_CreateMeshFromFaces)
+REGISTER_ECA_COMMAND(FECACommand_CreatePolygonMesh)
+REGISTER_ECA_COMMAND(FECACommand_MeshBoolean)
+REGISTER_ECA_COMMAND(FECACommand_MeshSelfUnion)
+REGISTER_ECA_COMMAND(FECACommand_MeshPlaneCut)
+REGISTER_ECA_COMMAND(FECACommand_MeshSlice)
+REGISTER_ECA_COMMAND(FECACommand_MeshMirror)
+REGISTER_ECA_COMMAND(FECACommand_MeshExtrude)
+REGISTER_ECA_COMMAND(FECACommand_ExtrudePolygon)
+REGISTER_ECA_COMMAND(FECACommand_RevolveProfile)
+REGISTER_ECA_COMMAND(FECACommand_SweepProfile)
+REGISTER_ECA_COMMAND(FECACommand_MeshInset)
+REGISTER_ECA_COMMAND(FECACommand_MeshOffset)
+REGISTER_ECA_COMMAND(FECACommand_MeshSimplify)
+REGISTER_ECA_COMMAND(FECACommand_MeshDisplaceFromDepth)
+REGISTER_ECA_COMMAND(FECACommand_MeshDisplaceFromHeightmap)
+REGISTER_ECA_COMMAND(FECACommand_MeshSubdivide)
+REGISTER_ECA_COMMAND(FECACommand_MeshRemesh)
+REGISTER_ECA_COMMAND(FECACommand_MeshSmooth)
+REGISTER_ECA_COMMAND(FECACommand_MeshDisplace)
+REGISTER_ECA_COMMAND(FECACommand_MeshBend)
+REGISTER_ECA_COMMAND(FECACommand_MeshTwist)
+REGISTER_ECA_COMMAND(FECACommand_CreateHeightfield)
+REGISTER_ECA_COMMAND(FECACommand_MeshGenerateUVs)
+REGISTER_ECA_COMMAND(FECACommand_MeshRecomputeNormals)
+REGISTER_ECA_COMMAND(FECACommand_MeshFlipNormals)
+REGISTER_ECA_COMMAND(FECACommand_MeshReverseWinding)
+REGISTER_ECA_COMMAND(FECACommand_MeshRepair)
+REGISTER_ECA_COMMAND(FECACommand_MeshSetMaterials)
+REGISTER_ECA_COMMAND(FECACommand_MeshTransform)
+REGISTER_ECA_COMMAND(FECACommand_GetMeshInfo)
+REGISTER_ECA_COMMAND(FECACommand_GetMeshBounds)
+REGISTER_ECA_COMMAND(FECACommand_SpawnProceduralMesh)
+REGISTER_ECA_COMMAND(FECACommand_UpdateProceduralMesh)
+REGISTER_ECA_COMMAND(FECACommand_MeshCombine)
+REGISTER_ECA_COMMAND(FECACommand_ImportMesh)
+REGISTER_ECA_COMMAND(FECACommand_ExportMesh)
+REGISTER_ECA_COMMAND(FECACommand_MeshVoronoiFracture)
+
+//==============================================================================
+// HELPER FUNCTIONS
+//==============================================================================
+
+namespace MeshHelpers
+{
+	// Load a static mesh from path
+	UStaticMesh* LoadStaticMesh(const FString& Path)
+	{
+		return LoadObject<UStaticMesh>(nullptr, *Path);
+	}
+
+	// Generate a cylinder mesh with proper UV seams (fixes texture wrapping issue)
+	// The default FCylinderGenerator has continuous UVs that wrap incorrectly at the seam
+	void GenerateCylinderWithUVSeam(FDynamicMesh3& OutMesh, double Radius, double Height, int32 RadialSegments, int32 HeightSegments, bool bCapped, bool bReverseOrientation)
+	{
+		OutMesh.Clear();
+		OutMesh.EnableAttributes();
+		OutMesh.Attributes()->SetNumUVLayers(1);
+
+		FDynamicMeshUVOverlay* UVOverlay = OutMesh.Attributes()->PrimaryUV();
+		// Normals will be computed by QuickComputeVertexNormals after mesh is built
+
+		// Generate vertices for cylinder body with duplicated seam vertices for proper UVs
+		// We need RadialSegments+1 vertices per ring to have U go from 0 to 1
+		int32 VertsPerRing = RadialSegments + 1;
+		int32 NumRings = HeightSegments + 1;
+
+		// Create body vertices and UVs
+		for (int32 Ring = 0; Ring < NumRings; Ring++)
+		{
+			double Z = -Height / 2.0 + (Height * Ring / HeightSegments);
+			double V = (double)Ring / HeightSegments;
+
+			for (int32 Seg = 0; Seg <= RadialSegments; Seg++)
+			{
+				double Angle = 2.0 * PI * Seg / RadialSegments;
+				double X = Radius * FMath::Cos(Angle);
+				double Y = Radius * FMath::Sin(Angle);
+
+				OutMesh.AppendVertex(FVector3d(X, Y, Z));
+
+				// UV: U goes from 0 to 1 around the cylinder, V goes from 0 to 1 along height
+				double U = (double)Seg / RadialSegments;
+				UVOverlay->AppendElement(FVector2f((float)U, (float)V));
+			}
+		}
+
+		// Create body triangles
+		for (int32 Ring = 0; Ring < HeightSegments; Ring++)
+		{
+			for (int32 Seg = 0; Seg < RadialSegments; Seg++)
+			{
+				int32 Current = Ring * VertsPerRing + Seg;
+				int32 Next = Current + 1;
+				int32 CurrentUp = Current + VertsPerRing;
+				int32 NextUp = Next + VertsPerRing;
+
+				// Two triangles per quad
+				int32 Tri1, Tri2;
+				if (bReverseOrientation)
+				{
+					Tri1 = OutMesh.AppendTriangle(Current, CurrentUp, Next);
+					Tri2 = OutMesh.AppendTriangle(Next, CurrentUp, NextUp);
+				}
+				else
+				{
+					Tri1 = OutMesh.AppendTriangle(Current, Next, CurrentUp);
+					Tri2 = OutMesh.AppendTriangle(Next, NextUp, CurrentUp);
+				}
+
+				// Set UVs for triangles
+				if (Tri1 >= 0)
+				{
+					if (bReverseOrientation)
+						UVOverlay->SetTriangle(Tri1, FIndex3i(Current, CurrentUp, Next));
+					else
+						UVOverlay->SetTriangle(Tri1, FIndex3i(Current, Next, CurrentUp));
+				}
+				if (Tri2 >= 0)
+				{
+					if (bReverseOrientation)
+						UVOverlay->SetTriangle(Tri2, FIndex3i(Next, CurrentUp, NextUp));
+					else
+						UVOverlay->SetTriangle(Tri2, FIndex3i(Next, NextUp, CurrentUp));
+				}
+			}
+		}
+
+		if (bCapped)
+		{
+			// Bottom cap center vertex
+			int32 BottomCenterVert = OutMesh.AppendVertex(FVector3d(0, 0, -Height / 2.0));
+			int32 BottomCenterUV = UVOverlay->AppendElement(FVector2f(0.5f, 0.5f));
+
+			// Bottom cap edge vertices (need separate vertices for hard edge normals)
+			TArray<int32> BottomEdgeVerts, BottomEdgeUVs;
+			for (int32 Seg = 0; Seg <= RadialSegments; Seg++)
+			{
+				double Angle = 2.0 * PI * Seg / RadialSegments;
+				double X = Radius * FMath::Cos(Angle);
+				double Y = Radius * FMath::Sin(Angle);
+
+				int32 VertIdx = OutMesh.AppendVertex(FVector3d(X, Y, -Height / 2.0));
+				BottomEdgeVerts.Add(VertIdx);
+
+				// Cap UVs are planar projected
+				float U = 0.5f + 0.5f * (float)FMath::Cos(Angle);
+				float V = 0.5f + 0.5f * (float)FMath::Sin(Angle);
+				BottomEdgeUVs.Add(UVOverlay->AppendElement(FVector2f(U, V)));
+			}
+
+			// Bottom cap triangles
+			for (int32 Seg = 0; Seg < RadialSegments; Seg++)
+			{
+				int32 Tri;
+				if (bReverseOrientation)
+					Tri = OutMesh.AppendTriangle(BottomCenterVert, BottomEdgeVerts[Seg], BottomEdgeVerts[Seg + 1]);
+				else
+					Tri = OutMesh.AppendTriangle(BottomCenterVert, BottomEdgeVerts[Seg + 1], BottomEdgeVerts[Seg]);
+
+				if (Tri >= 0)
+				{
+					if (bReverseOrientation)
+						UVOverlay->SetTriangle(Tri, FIndex3i(BottomCenterUV, BottomEdgeUVs[Seg], BottomEdgeUVs[Seg + 1]));
+					else
+						UVOverlay->SetTriangle(Tri, FIndex3i(BottomCenterUV, BottomEdgeUVs[Seg + 1], BottomEdgeUVs[Seg]));
+				}
+			}
+
+			// Top cap center vertex
+			int32 TopCenterVert = OutMesh.AppendVertex(FVector3d(0, 0, Height / 2.0));
+			int32 TopCenterUV = UVOverlay->AppendElement(FVector2f(0.5f, 0.5f));
+
+			// Top cap edge vertices
+			TArray<int32> TopEdgeVerts, TopEdgeUVs;
+			for (int32 Seg = 0; Seg <= RadialSegments; Seg++)
+			{
+				double Angle = 2.0 * PI * Seg / RadialSegments;
+				double X = Radius * FMath::Cos(Angle);
+				double Y = Radius * FMath::Sin(Angle);
+
+				int32 VertIdx = OutMesh.AppendVertex(FVector3d(X, Y, Height / 2.0));
+				TopEdgeVerts.Add(VertIdx);
+
+				float U = 0.5f + 0.5f * (float)FMath::Cos(Angle);
+				float V = 0.5f + 0.5f * (float)FMath::Sin(Angle);
+				TopEdgeUVs.Add(UVOverlay->AppendElement(FVector2f(U, V)));
+			}
+
+			// Top cap triangles
+			for (int32 Seg = 0; Seg < RadialSegments; Seg++)
+			{
+				int32 Tri;
+				if (bReverseOrientation)
+					Tri = OutMesh.AppendTriangle(TopCenterVert, TopEdgeVerts[Seg + 1], TopEdgeVerts[Seg]);
+				else
+					Tri = OutMesh.AppendTriangle(TopCenterVert, TopEdgeVerts[Seg], TopEdgeVerts[Seg + 1]);
+
+				if (Tri >= 0)
+				{
+					if (bReverseOrientation)
+						UVOverlay->SetTriangle(Tri, FIndex3i(TopCenterUV, TopEdgeUVs[Seg + 1], TopEdgeUVs[Seg]));
+					else
+						UVOverlay->SetTriangle(Tri, FIndex3i(TopCenterUV, TopEdgeUVs[Seg], TopEdgeUVs[Seg + 1]));
+				}
+			}
+		}
+	}
+
+	// Create a new static mesh asset from FDynamicMesh3
+	UStaticMesh* CreateStaticMeshFromDynamicMesh(
+		const FDynamicMesh3& DynamicMesh,
+		const FString& AssetPath,
+		bool bGenerateCollision = true,
+		bool bGenerateLightmapUVs = true,
+		FString* OutError = nullptr)
+	{
+		// Extract asset name and package path
+		FString PackagePath, AssetName;
+		AssetPath.Split(TEXT("/"), &PackagePath, &AssetName, ESearchCase::IgnoreCase, ESearchDir::FromEnd);
+		if (AssetName.IsEmpty())
+		{
+			AssetName = FPaths::GetBaseFilename(AssetPath);
+			PackagePath = AssetPath;
+		}
+		
+		// Create package
+		UPackage* Package = CreatePackage(*AssetPath);
+		if (!Package)
+		{
+			if (OutError) *OutError = TEXT("Failed to create package");
+			return nullptr;
+		}
+
+		// Create static mesh
+		UStaticMesh* StaticMesh = NewObject<UStaticMesh>(Package, *AssetName, RF_Public | RF_Standalone);
+		if (!StaticMesh)
+		{
+			if (OutError) *OutError = TEXT("Failed to create static mesh");
+			return nullptr;
+		}
+
+		// Convert DynamicMesh to MeshDescription
+		FMeshDescription MeshDesc;
+		FStaticMeshAttributes Attributes(MeshDesc);
+		Attributes.Register();
+
+		FDynamicMeshToMeshDescription Converter;
+		Converter.Convert(&DynamicMesh, MeshDesc);
+
+		// Set up material slots
+		int32 NumMaterialIDs = 1;
+		if (DynamicMesh.HasAttributes() && DynamicMesh.Attributes()->HasMaterialID())
+		{
+			const FDynamicMeshMaterialAttribute* MaterialAttr = DynamicMesh.Attributes()->GetMaterialID();
+			for (int32 TriID : DynamicMesh.TriangleIndicesItr())
+			{
+				NumMaterialIDs = FMath::Max(NumMaterialIDs, MaterialAttr->GetValue(TriID) + 1);
+			}
+		}
+
+		StaticMesh->GetStaticMaterials().Empty();
+		for (int32 i = 0; i < NumMaterialIDs; i++)
+		{
+			FStaticMaterial NewMaterial(nullptr, FName(*FString::Printf(TEXT("Material_%d"), i)), FName(*FString::Printf(TEXT("Material_%d"), i)));
+			NewMaterial.UVChannelData = FMeshUVChannelInfo(1.0f);  // Initialize UV channel data to prevent ensure failures
+			StaticMesh->GetStaticMaterials().Add(NewMaterial);
+		}
+
+		// Add LOD 0
+		FStaticMeshSourceModel& SourceModel = StaticMesh->AddSourceModel();
+		SourceModel.BuildSettings.bRecomputeNormals = !DynamicMesh.HasAttributes() || !DynamicMesh.Attributes()->PrimaryNormals();
+		SourceModel.BuildSettings.bRecomputeTangents = true;
+		SourceModel.BuildSettings.bGenerateLightmapUVs = bGenerateLightmapUVs;
+		SourceModel.BuildSettings.SrcLightmapIndex = 0;
+		SourceModel.BuildSettings.DstLightmapIndex = 1;
+
+		// Set mesh description
+		FMeshDescription* MeshDescription = StaticMesh->CreateMeshDescription(0);
+		*MeshDescription = MoveTemp(MeshDesc);
+		StaticMesh->CommitMeshDescription(0);
+
+		// Set up collision
+		if (bGenerateCollision)
+		{
+			StaticMesh->CreateBodySetup();
+			if (StaticMesh->GetBodySetup())
+			{
+				StaticMesh->GetBodySetup()->CollisionTraceFlag = CTF_UseComplexAsSimple;
+			}
+		}
+
+		// Build the mesh
+		StaticMesh->SetImportVersion(EImportStaticMeshVersion::LastVersion);
+		StaticMesh->Build(false);
+		StaticMesh->PostEditChange();
+
+		// Register with asset registry
+		FAssetRegistryModule::AssetCreated(StaticMesh);
+		Package->MarkPackageDirty();
+
+		return StaticMesh;
+	}
+
+	// Load static mesh into FDynamicMesh3
+	bool LoadStaticMeshIntoDynamicMesh(UStaticMesh* StaticMesh, FDynamicMesh3& OutMesh, FString* OutError = nullptr)
+	{
+		if (!StaticMesh)
+		{
+			if (OutError) *OutError = TEXT("Static mesh is null");
+			return false;
+		}
+
+		// Get mesh description for LOD 0
+		const FMeshDescription* MeshDesc = StaticMesh->GetMeshDescription(0);
+		if (!MeshDesc)
+		{
+			if (OutError) *OutError = TEXT("Mesh has no mesh description");
+			return false;
+		}
+
+		// Convert to dynamic mesh
+		FMeshDescriptionToDynamicMesh Converter;
+		Converter.Convert(MeshDesc, OutMesh);
+
+		return true;
+	}
+
+	// Save dynamic mesh back to static mesh asset
+	bool SaveDynamicMeshToStaticMesh(const FDynamicMesh3& DynamicMesh, UStaticMesh* StaticMesh, bool bRebuild = true)
+	{
+		if (!StaticMesh)
+		{
+			return false;
+		}
+
+		// Convert to mesh description
+		FMeshDescription MeshDesc;
+		FStaticMeshAttributes Attributes(MeshDesc);
+		Attributes.Register();
+
+		FDynamicMeshToMeshDescription Converter;
+		Converter.Convert(&DynamicMesh, MeshDesc);
+
+		// Update static mesh
+		FMeshDescription* ExistingMeshDesc = StaticMesh->GetMeshDescription(0);
+		if (!ExistingMeshDesc)
+		{
+			ExistingMeshDesc = StaticMesh->CreateMeshDescription(0);
+		}
+		*ExistingMeshDesc = MoveTemp(MeshDesc);
+		StaticMesh->CommitMeshDescription(0);
+
+		if (bRebuild)
+		{
+			StaticMesh->Build(false);
+			StaticMesh->PostEditChange();
+			StaticMesh->MarkPackageDirty();
+		}
+
+		return true;
+	}
+
+	// Parse a vector from JSON object
+	FVector3d ParseVector(const TSharedPtr<FJsonObject>& Obj, FVector3d Default = FVector3d::Zero())
+	{
+		if (!Obj.IsValid())
+		{
+			return Default;
+		}
+		return FVector3d(
+			Obj->GetNumberField(TEXT("x")),
+			Obj->GetNumberField(TEXT("y")),
+			Obj->GetNumberField(TEXT("z"))
+		);
+	}
+
+	// Parse 2D point from JSON object
+	FVector2d ParseVector2D(const TSharedPtr<FJsonObject>& Obj, FVector2d Default = FVector2d::Zero())
+	{
+		if (!Obj.IsValid())
+		{
+			return Default;
+		}
+		return FVector2d(
+			Obj->GetNumberField(TEXT("x")),
+			Obj->GetNumberField(TEXT("y"))
+		);
+	}
+
+	// Parse transform from JSON object
+	FTransform ParseTransform(const TSharedPtr<FJsonObject>& Obj)
+	{
+		FTransform Result = FTransform::Identity;
+		if (!Obj.IsValid())
+		{
+			return Result;
+		}
+
+		const TSharedPtr<FJsonObject>* LocationObj;
+		if (Obj->TryGetObjectField(TEXT("location"), LocationObj))
+		{
+			FVector Location(
+				(*LocationObj)->GetNumberField(TEXT("x")),
+				(*LocationObj)->GetNumberField(TEXT("y")),
+				(*LocationObj)->GetNumberField(TEXT("z"))
+			);
+			Result.SetTranslation(Location);
+		}
+
+		const TSharedPtr<FJsonObject>* RotationObj;
+		if (Obj->TryGetObjectField(TEXT("rotation"), RotationObj))
+		{
+			double Pitch = (*RotationObj)->GetNumberField(TEXT("pitch"));
+			double Yaw = (*RotationObj)->GetNumberField(TEXT("yaw"));
+			double Roll = (*RotationObj)->GetNumberField(TEXT("roll"));
+			Result.SetRotation(FQuat(FRotator(Pitch, Yaw, Roll)));
+		}
+
+		const TSharedPtr<FJsonObject>* ScaleObj;
+		if (Obj->TryGetObjectField(TEXT("scale"), ScaleObj))
+		{
+			FVector Scale(
+				(*ScaleObj)->GetNumberField(TEXT("x")),
+				(*ScaleObj)->GetNumberField(TEXT("y")),
+				(*ScaleObj)->GetNumberField(TEXT("z"))
+			);
+			Result.SetScale3D(Scale);
+		}
+
+		return Result;
+	}
+
+	// Build mesh result JSON
+	TSharedPtr<FJsonObject> BuildMeshResultJson(UStaticMesh* Mesh, int32 VertexCount = -1, int32 TriangleCount = -1)
+	{
+		TSharedPtr<FJsonObject> Result = MakeShared<FJsonObject>();
+		
+		if (Mesh)
+		{
+			Result->SetStringField(TEXT("mesh_path"), Mesh->GetPathName());
+			Result->SetStringField(TEXT("mesh_name"), Mesh->GetName());
+
+			// Get actual counts from mesh if not provided
+			if (VertexCount < 0 || TriangleCount < 0)
+			{
+				const FMeshDescription* MeshDesc = Mesh->GetMeshDescription(0);
+				if (MeshDesc)
+				{
+					VertexCount = MeshDesc->Vertices().Num();
+					TriangleCount = MeshDesc->Triangles().Num();
+				}
+			}
+
+			Result->SetNumberField(TEXT("vertex_count"), VertexCount);
+			Result->SetNumberField(TEXT("triangle_count"), TriangleCount);
+			Result->SetNumberField(TEXT("material_count"), Mesh->GetStaticMaterials().Num());
+
+			// Bounds
+			FBoxSphereBounds Bounds = Mesh->GetBounds();
+			TSharedPtr<FJsonObject> BoundsObj = MakeShared<FJsonObject>();
+			BoundsObj->SetNumberField(TEXT("size_x"), Bounds.BoxExtent.X * 2);
+			BoundsObj->SetNumberField(TEXT("size_y"), Bounds.BoxExtent.Y * 2);
+			BoundsObj->SetNumberField(TEXT("size_z"), Bounds.BoxExtent.Z * 2);
+			BoundsObj->SetNumberField(TEXT("radius"), Bounds.SphereRadius);
+			Result->SetObjectField(TEXT("bounds"), BoundsObj);
+		}
+
+		return Result;
+	}
+
+	// Compute vertex normals by averaging face normals across colocated vertices
+	// This properly handles vertices at the same position but with different indices
+	// (e.g., UV seam vertices, hard edge splits)
+	// 
+	// Parameters:
+	// - DynMesh: The dynamic mesh to compute normals for (modified in place)
+	// - WeldTolerance: Maximum distance to consider vertices as colocated (default: 0.0001)
+	// - bWeightByArea: If true, weight face normals by triangle area (default: true)
+	// - bWeightByAngle: If true, weight face normals by angle at vertex (default: true)
+	void ComputeVertexNormalsAcrossSeams(FDynamicMesh3& DynMesh, double WeldTolerance = 0.0001, bool bWeightByArea = true, bool bWeightByAngle = true)
+	{
+		// Ensure mesh has normal overlay
+		if (!DynMesh.HasAttributes())
+		{
+			DynMesh.EnableAttributes();
+		}
+		FDynamicMeshNormalOverlay* NormalOverlay = DynMesh.Attributes()->PrimaryNormals();
+		if (!NormalOverlay)
+		{
+			return;
+		}
+		
+		// Clear existing normals
+		NormalOverlay->ClearElements();
+		
+		// Step 1: Compute face normals and areas for all triangles
+		TMap<int32, FVector3d> FaceNormals;
+		TMap<int32, double> FaceAreas;
+		
+		for (int32 TriID : DynMesh.TriangleIndicesItr())
+		{
+			FVector3d V0, V1, V2;
+			DynMesh.GetTriVertices(TriID, V0, V1, V2);
+			
+			FVector3d Edge1 = V1 - V0;
+			FVector3d Edge2 = V2 - V0;
+			// Note: Edge2.Cross(Edge1) gives correct normal direction for UE's coordinate system
+			// (counter-clockwise winding with outward-facing normals)
+			FVector3d CrossProduct = Edge2.Cross(Edge1);
+			
+			double Area = CrossProduct.Length() * 0.5;
+			FVector3d Normal = Area > UE_DOUBLE_KINDA_SMALL_NUMBER ? CrossProduct.GetSafeNormal() : FVector3d::UnitZ();
+			
+			FaceNormals.Add(TriID, Normal);
+			FaceAreas.Add(TriID, Area);
+		}
+		
+		// Step 2: Build spatial hash of vertices by position
+		// This groups vertices that are at the same position (within tolerance)
+		TMap<FIntVector, TArray<int32>> SpatialHash;
+		double CellSize = FMath::Max(WeldTolerance * 2.0, 0.001);
+		
+		for (int32 VertID : DynMesh.VertexIndicesItr())
+		{
+			FVector3d Pos = DynMesh.GetVertex(VertID);
+			FIntVector Cell(
+				FMath::FloorToInt(Pos.X / CellSize),
+				FMath::FloorToInt(Pos.Y / CellSize),
+				FMath::FloorToInt(Pos.Z / CellSize)
+			);
+			SpatialHash.FindOrAdd(Cell).Add(VertID);
+		}
+		
+		// Step 3: Build map of vertex position -> all colocated vertex IDs
+		TMap<int32, TArray<int32>> ColocatedVertices; // VertID -> all verts at same position
+		TSet<int32> Processed;
+		
+		for (int32 VertID : DynMesh.VertexIndicesItr())
+		{
+			if (Processed.Contains(VertID))
+			{
+				continue;
+			}
+			
+			FVector3d Pos = DynMesh.GetVertex(VertID);
+			FIntVector Cell(
+				FMath::FloorToInt(Pos.X / CellSize),
+				FMath::FloorToInt(Pos.Y / CellSize),
+				FMath::FloorToInt(Pos.Z / CellSize)
+			);
+			
+			TArray<int32> Colocated;
+			
+			// Check this cell and all 26 neighbors
+			for (int32 dx = -1; dx <= 1; dx++)
+			{
+				for (int32 dy = -1; dy <= 1; dy++)
+				{
+					for (int32 dz = -1; dz <= 1; dz++)
+					{
+						FIntVector NeighborCell = Cell + FIntVector(dx, dy, dz);
+						if (TArray<int32>* CellVerts = SpatialHash.Find(NeighborCell))
+						{
+							for (int32 OtherVertID : *CellVerts)
+							{
+								FVector3d OtherPos = DynMesh.GetVertex(OtherVertID);
+								if (FVector3d::DistSquared(Pos, OtherPos) <= WeldTolerance * WeldTolerance)
+								{
+									Colocated.AddUnique(OtherVertID);
+								}
+							}
+						}
+					}
+				}
+			}
+			
+			// Store the colocated set for all vertices in the group
+			for (int32 ColocVertID : Colocated)
+			{
+				ColocatedVertices.Add(ColocVertID, Colocated);
+				Processed.Add(ColocVertID);
+			}
+		}
+		
+		// Step 4: For each vertex, accumulate weighted face normals from ALL colocated vertices
+		TMap<int32, FVector3d> AccumulatedNormals;
+		
+		for (int32 VertID : DynMesh.VertexIndicesItr())
+		{
+			FVector3d AccumNormal = FVector3d::Zero();
+			
+			// Get all vertices at this position
+			TArray<int32>* ColocatedPtr = ColocatedVertices.Find(VertID);
+			TArray<int32> AllVerts = ColocatedPtr ? *ColocatedPtr : TArray<int32>{VertID};
+			
+			// For each colocated vertex, accumulate face normals
+			for (int32 ColocVertID : AllVerts)
+			{
+				// Get all triangles that use this vertex
+				for (int32 TriID : DynMesh.VtxTrianglesItr(ColocVertID))
+				{
+					FVector3d FaceNormal = FaceNormals[TriID];
+					double Weight = 1.0;
+					
+					// Weight by triangle area
+					if (bWeightByArea)
+					{
+						Weight *= FaceAreas[TriID];
+					}
+					
+					// Weight by angle at vertex
+					if (bWeightByAngle)
+					{
+						FIndex3i TriVerts = DynMesh.GetTriangle(TriID);
+						FVector3d V0 = DynMesh.GetVertex(TriVerts.A);
+						FVector3d V1 = DynMesh.GetVertex(TriVerts.B);
+						FVector3d V2 = DynMesh.GetVertex(TriVerts.C);
+						
+						// Find which corner this vertex is at
+						FVector3d VertPos = DynMesh.GetVertex(ColocVertID);
+						FVector3d Edge1, Edge2;
+						
+						if (FVector3d::DistSquared(V0, VertPos) < WeldTolerance * WeldTolerance)
+						{
+							Edge1 = (V1 - V0).GetSafeNormal();
+							Edge2 = (V2 - V0).GetSafeNormal();
+						}
+						else if (FVector3d::DistSquared(V1, VertPos) < WeldTolerance * WeldTolerance)
+						{
+							Edge1 = (V0 - V1).GetSafeNormal();
+							Edge2 = (V2 - V1).GetSafeNormal();
+						}
+						else
+						{
+							Edge1 = (V0 - V2).GetSafeNormal();
+							Edge2 = (V1 - V2).GetSafeNormal();
+						}
+						
+						double Angle = FMath::Acos(FMath::Clamp(Edge1.Dot(Edge2), -1.0, 1.0));
+						Weight *= Angle;
+					}
+					
+					AccumNormal += FaceNormal * Weight;
+				}
+			}
+			
+			// Normalize the accumulated normal
+			if (AccumNormal.SquaredLength() > UE_DOUBLE_KINDA_SMALL_NUMBER)
+			{
+				AccumNormal.Normalize();
+			}
+			else
+			{
+				AccumNormal = FVector3d::UnitZ();
+			}
+			
+			AccumulatedNormals.Add(VertID, AccumNormal);
+		}
+		
+		// Step 5: Set normals on the overlay
+		// Create normal elements for each vertex and assign to triangles
+		TMap<int32, int32> VertexToNormalElement;
+		
+		for (const auto& Pair : AccumulatedNormals)
+		{
+			int32 VertID = Pair.Key;
+			FVector3d Normal = Pair.Value;
+			int32 NormalElemID = NormalOverlay->AppendElement(FVector3f(Normal));
+			VertexToNormalElement.Add(VertID, NormalElemID);
+		}
+		
+		// Assign normals to triangles
+		for (int32 TriID : DynMesh.TriangleIndicesItr())
+		{
+			FIndex3i TriVerts = DynMesh.GetTriangle(TriID);
+			
+			int32* NormA = VertexToNormalElement.Find(TriVerts.A);
+			int32* NormB = VertexToNormalElement.Find(TriVerts.B);
+			int32* NormC = VertexToNormalElement.Find(TriVerts.C);
+			
+			if (NormA && NormB && NormC)
+			{
+				NormalOverlay->SetTriangle(TriID, FIndex3i(*NormA, *NormB, *NormC));
+			}
+		}
+	}
+
+}
+
+//==============================================================================
+// PRIMITIVE CREATION
+//==============================================================================
+
+FECACommandResult FECACommand_CreatePrimitive::Execute(const TSharedPtr<FJsonObject>& Params)
+{
+	FString AssetPath;
+	if (!GetStringParam(Params, TEXT("asset_path"), AssetPath))
+	{
+		return FECACommandResult::Error(TEXT("Missing required parameter: asset_path"));
+	}
+
+	FString Type;
+	if (!GetStringParam(Params, TEXT("type"), Type))
+	{
+		return FECACommandResult::Error(TEXT("Missing required parameter: type"));
+	}
+
+	// Get dimensions
+	const TSharedPtr<FJsonObject>* DimensionsObj = nullptr;
+	Params->TryGetObjectField(TEXT("dimensions"), DimensionsObj);
+
+	// Get segments
+	const TSharedPtr<FJsonObject>* SegmentsObj = nullptr;
+	Params->TryGetObjectField(TEXT("segments"), SegmentsObj);
+
+	FString Origin = TEXT("center");
+	GetStringParam(Params, TEXT("origin"), Origin, false);
+
+	bool bGenerateCollision = true;
+	GetBoolParam(Params, TEXT("generate_collision"), bGenerateCollision, false);
+
+	// Create the primitive mesh
+	FDynamicMesh3 DynMesh;
+	DynMesh.EnableAttributes();
+	DynMesh.Attributes()->EnablePrimaryColors();
+	DynMesh.Attributes()->SetNumUVLayers(1);
+
+	Type = Type.ToLower();
+
+	if (Type == TEXT("box"))
+	{
+		double SizeX = DimensionsObj ? (*DimensionsObj)->GetNumberField(TEXT("x")) : 100.0;
+		double SizeY = DimensionsObj ? (*DimensionsObj)->GetNumberField(TEXT("y")) : 100.0;
+		double SizeZ = DimensionsObj ? (*DimensionsObj)->GetNumberField(TEXT("z")) : 100.0;
+
+		FGridBoxMeshGenerator BoxGen;
+		BoxGen.Box = FOrientedBox3d(FVector3d::Zero(), FVector3d(SizeX / 2, SizeY / 2, SizeZ / 2));
+		BoxGen.EdgeVertices = FIndex3i(0, 0, 0);
+		if (SegmentsObj)
+		{
+			BoxGen.EdgeVertices.A = (*SegmentsObj)->GetIntegerField(TEXT("x"));
+			BoxGen.EdgeVertices.B = (*SegmentsObj)->GetIntegerField(TEXT("y"));
+			BoxGen.EdgeVertices.C = (*SegmentsObj)->GetIntegerField(TEXT("z"));
+		}
+		BoxGen.bPolygroupPerQuad = false;
+
+		BoxGen.Generate();
+		DynMesh.Copy(&BoxGen);
+
+		if (Origin == TEXT("base"))
+		{
+			MeshTransforms::Translate(DynMesh, FVector3d(0, 0, SizeZ / 2));
+		}
+		else if (Origin == TEXT("corner"))
+		{
+			MeshTransforms::Translate(DynMesh, FVector3d(SizeX / 2, SizeY / 2, SizeZ / 2));
+		}
+	}
+	else if (Type == TEXT("sphere"))
+	{
+		double Radius = DimensionsObj ? (*DimensionsObj)->GetNumberField(TEXT("radius")) : 50.0;
+		if (Radius <= 0 && DimensionsObj)
+		{
+			Radius = (*DimensionsObj)->GetNumberField(TEXT("x")) / 2;
+		}
+
+		int32 NumSteps = 32;
+		if (SegmentsObj)
+		{
+			NumSteps = (*SegmentsObj)->GetIntegerField(TEXT("radial"));
+			if (NumSteps <= 0) NumSteps = 32;
+		}
+
+		FSphereGenerator SphereGen;
+		SphereGen.Radius = Radius;
+		SphereGen.NumPhi = NumSteps;
+		SphereGen.NumTheta = NumSteps;
+		SphereGen.bPolygroupPerQuad = false;
+
+		SphereGen.Generate();
+		DynMesh.Copy(&SphereGen);
+
+		if (Origin == TEXT("base"))
+		{
+			MeshTransforms::Translate(DynMesh, FVector3d(0, 0, Radius));
+		}
+	}
+	else if (Type == TEXT("cylinder"))
+	{
+		double Radius = DimensionsObj ? (*DimensionsObj)->GetNumberField(TEXT("radius")) : 50.0;
+		double Height = DimensionsObj ? (*DimensionsObj)->GetNumberField(TEXT("height")) : 100.0;
+
+		int32 RadialSegments = 32;
+		int32 HeightSegments = 1;
+		if (SegmentsObj)
+		{
+			RadialSegments = (*SegmentsObj)->GetIntegerField(TEXT("radial"));
+			HeightSegments = (*SegmentsObj)->GetIntegerField(TEXT("height"));
+			if (RadialSegments <= 0) RadialSegments = 32;
+			if (HeightSegments <= 0) HeightSegments = 1;
+		}
+
+		// Use custom cylinder generator with proper UV seams
+		MeshHelpers::GenerateCylinderWithUVSeam(DynMesh, Radius, Height, RadialSegments, HeightSegments, true, true);
+
+		if (Origin == TEXT("center"))
+		{
+			// Cylinder is already centered by default
+		}
+		else if (Origin == TEXT("base"))
+		{
+			MeshTransforms::Translate(DynMesh, FVector3d(0, 0, Height / 2));
+		}
+	}
+	else if (Type == TEXT("cone"))
+	{
+		double Radius = DimensionsObj ? (*DimensionsObj)->GetNumberField(TEXT("radius")) : 50.0;
+		double Height = DimensionsObj ? (*DimensionsObj)->GetNumberField(TEXT("height")) : 100.0;
+
+		int32 RadialSegments = 32;
+		if (SegmentsObj)
+		{
+			RadialSegments = (*SegmentsObj)->GetIntegerField(TEXT("radial"));
+			if (RadialSegments <= 0) RadialSegments = 32;
+		}
+
+		FCylinderGenerator ConeGen;
+		ConeGen.Radius[0] = 0.001; // Tip
+		ConeGen.Radius[1] = Radius;
+		ConeGen.Height = Height;
+		ConeGen.AngleSamples = RadialSegments;
+		ConeGen.LengthSamples = 1;
+		ConeGen.bCapped = true;
+		ConeGen.bPolygroupPerQuad = false;
+
+		ConeGen.Generate();
+		DynMesh.Copy(&ConeGen);
+
+		if (Origin == TEXT("center"))
+		{
+			MeshTransforms::Translate(DynMesh, FVector3d(0, 0, -Height / 2));
+		}
+	}
+	else if (Type == TEXT("capsule"))
+	{
+		double Radius = DimensionsObj ? (*DimensionsObj)->GetNumberField(TEXT("radius")) : 25.0;
+		double Height = DimensionsObj ? (*DimensionsObj)->GetNumberField(TEXT("height")) : 100.0;
+
+		int32 RadialSegments = 32;
+		if (SegmentsObj)
+		{
+			RadialSegments = (*SegmentsObj)->GetIntegerField(TEXT("radial"));
+			if (RadialSegments <= 0) RadialSegments = 32;
+		}
+
+		FCapsuleGenerator CapsuleGen;
+		CapsuleGen.Radius = Radius;
+		CapsuleGen.SegmentLength = FMath::Max(0.0, Height - 2 * Radius);
+		CapsuleGen.NumHemisphereArcSteps = RadialSegments / 4;
+		CapsuleGen.NumCircleSteps = RadialSegments;
+		CapsuleGen.bPolygroupPerQuad = false;
+
+		CapsuleGen.Generate();
+		DynMesh.Copy(&CapsuleGen);
+
+		if (Origin == TEXT("base"))
+		{
+			MeshTransforms::Translate(DynMesh, FVector3d(0, 0, Height / 2));
+		}
+	}
+	else if (Type == TEXT("torus"))
+	{
+		double MajorRadius = DimensionsObj ? (*DimensionsObj)->GetNumberField(TEXT("major_radius")) : 50.0;
+		double MinorRadius = DimensionsObj ? (*DimensionsObj)->GetNumberField(TEXT("minor_radius")) : 15.0;
+
+		int32 MajorSegments = 32;
+		int32 MinorSegments = 16;
+		if (SegmentsObj)
+		{
+			MajorSegments = (*SegmentsObj)->GetIntegerField(TEXT("major"));
+			MinorSegments = (*SegmentsObj)->GetIntegerField(TEXT("minor"));
+			if (MajorSegments <= 0) MajorSegments = 32;
+			if (MinorSegments <= 0) MinorSegments = 16;
+		}
+
+		// Generate torus using revolve - create a circle profile and revolve it
+		TArray<FVector2d> ProfilePoints;
+		for (int32 i = 0; i <= MinorSegments; i++)
+		{
+			double Angle = 2.0 * PI * i / MinorSegments;
+			ProfilePoints.Add(FVector2d(
+				MajorRadius + MinorRadius * FMath::Cos(Angle),
+				MinorRadius * FMath::Sin(Angle)
+			));
+		}
+
+		FRevolvePlanarPathGenerator TorusGen;
+		TorusGen.PathVertices = ProfilePoints;
+		TorusGen.RevolveDegrees = 360.0f;
+		TorusGen.Steps = MajorSegments;
+		TorusGen.bCapped = false;
+		FDynamicMesh3 TorusMesh = TorusGen.GenerateMesh();
+		DynMesh.Copy(TorusMesh);
+
+	}
+	else if (Type == TEXT("plane"))
+	{
+		double SizeX = DimensionsObj ? (*DimensionsObj)->GetNumberField(TEXT("x")) : 100.0;
+		double SizeY = DimensionsObj ? (*DimensionsObj)->GetNumberField(TEXT("y")) : 100.0;
+
+		int32 SubdivisionsX = 1;
+		int32 SubdivisionsY = 1;
+		if (SegmentsObj)
+		{
+			SubdivisionsX = (*SegmentsObj)->GetIntegerField(TEXT("x"));
+			SubdivisionsY = (*SegmentsObj)->GetIntegerField(TEXT("y"));
+			if (SubdivisionsX <= 0) SubdivisionsX = 1;
+			if (SubdivisionsY <= 0) SubdivisionsY = 1;
+		}
+
+		FRectangleMeshGenerator PlaneGen;
+		PlaneGen.Width = SizeX;
+		PlaneGen.Height = SizeY;
+		PlaneGen.WidthVertexCount = SubdivisionsX + 1;
+		PlaneGen.HeightVertexCount = SubdivisionsY + 1;
+		PlaneGen.bSinglePolyGroup = true;
+
+		PlaneGen.Generate();
+		DynMesh.Copy(&PlaneGen);
+
+		if (Origin == TEXT("corner"))
+		{
+			MeshTransforms::Translate(DynMesh, FVector3d(SizeX / 2, SizeY / 2, 0));
+		}
+	}
+	else if (Type == TEXT("disc"))
+	{
+		double Radius = DimensionsObj ? (*DimensionsObj)->GetNumberField(TEXT("radius")) : 50.0;
+		double InnerRadius = DimensionsObj ? (*DimensionsObj)->GetNumberField(TEXT("inner_radius")) : 0.0;
+
+		int32 RadialSegments = 32;
+		if (SegmentsObj)
+		{
+			RadialSegments = (*SegmentsObj)->GetIntegerField(TEXT("radial"));
+			if (RadialSegments <= 0) RadialSegments = 32;
+		}
+
+		if (InnerRadius > 0)
+		{
+			// Use punctured disc for ring shape
+			FPuncturedDiscMeshGenerator DiscGen;
+			DiscGen.Radius = Radius;
+			DiscGen.HoleRadius = InnerRadius;
+			DiscGen.AngleSamples = RadialSegments;
+			DiscGen.RadialSamples = 1;
+			DiscGen.bSinglePolygroup = true;
+
+			DiscGen.Generate();
+			DynMesh.Copy(&DiscGen);
+		}
+		else
+		{
+			FDiscMeshGenerator DiscGen;
+			DiscGen.Radius = Radius;
+			DiscGen.AngleSamples = RadialSegments;
+			DiscGen.RadialSamples = 1;
+			DiscGen.bSinglePolygroup = true;
+
+			DiscGen.Generate();
+			DynMesh.Copy(&DiscGen);
+		}
+	}
+	else if (Type == TEXT("grid"))
+	{
+		double SizeX = DimensionsObj ? (*DimensionsObj)->GetNumberField(TEXT("x")) : 100.0;
+		double SizeY = DimensionsObj ? (*DimensionsObj)->GetNumberField(TEXT("y")) : 100.0;
+		double SizeZ = DimensionsObj ? (*DimensionsObj)->GetNumberField(TEXT("z")) : 100.0;
+
+		int32 CellsX = 2;
+		int32 CellsY = 2;
+		int32 CellsZ = 2;
+		if (SegmentsObj)
+		{
+			CellsX = (*SegmentsObj)->GetIntegerField(TEXT("x"));
+			CellsY = (*SegmentsObj)->GetIntegerField(TEXT("y"));
+			CellsZ = (*SegmentsObj)->GetIntegerField(TEXT("z"));
+			if (CellsX <= 0) CellsX = 2;
+			if (CellsY <= 0) CellsY = 2;
+			if (CellsZ <= 0) CellsZ = 2;
+		}
+
+		FGridBoxMeshGenerator GridGen;
+		GridGen.Box = FOrientedBox3d(FVector3d::Zero(), FVector3d(SizeX / 2, SizeY / 2, SizeZ / 2));
+		GridGen.EdgeVertices = FIndex3i(CellsX, CellsY, CellsZ);
+		GridGen.bPolygroupPerQuad = false;
+
+		GridGen.Generate();
+		DynMesh.Copy(&GridGen);
+
+		if (Origin == TEXT("base"))
+		{
+			MeshTransforms::Translate(DynMesh, FVector3d(0, 0, SizeZ / 2));
+		}
+		else if (Origin == TEXT("corner"))
+		{
+			MeshTransforms::Translate(DynMesh, FVector3d(SizeX / 2, SizeY / 2, SizeZ / 2));
+		}
+	}
+	else
+	{
+		return FECACommandResult::Error(FString::Printf(TEXT("Unknown primitive type: %s"), *Type));
+	}
+
+	// Compute proper vertex normals for the mesh
+	FMeshNormals::QuickComputeVertexNormals(DynMesh);
+
+	// Create the static mesh asset
+	FString Error;
+	UStaticMesh* StaticMesh = MeshHelpers::CreateStaticMeshFromDynamicMesh(DynMesh, AssetPath, bGenerateCollision, true, &Error);
+	
+	if (!StaticMesh)
+	{
+		return FECACommandResult::Error(Error);
+	}
+
+	// Apply material if specified
+	FString MaterialPath;
+	if (GetStringParam(Params, TEXT("material_path"), MaterialPath, false))
+	{
+		UMaterialInterface* Material = LoadObject<UMaterialInterface>(nullptr, *MaterialPath);
+		if (Material)
+		{
+			StaticMesh->SetMaterial(0, Material);
+		}
+	}
+
+	TSharedPtr<FJsonObject> Result = MeshHelpers::BuildMeshResultJson(StaticMesh, DynMesh.VertexCount(), DynMesh.TriangleCount());
+	Result->SetStringField(TEXT("type"), Type);
+	
+	return FECACommandResult::Success(Result);
+}
+
+//==============================================================================
+// CREATE MESH FROM DATA
+//==============================================================================
+
+FECACommandResult FECACommand_CreateMesh::Execute(const TSharedPtr<FJsonObject>& Params)
+{
+	FString AssetPath;
+	if (!GetStringParam(Params, TEXT("asset_path"), AssetPath))
+	{
+		return FECACommandResult::Error(TEXT("Missing required parameter: asset_path"));
+	}
+
+	const TArray<TSharedPtr<FJsonValue>>* VerticesArray;
+	if (!GetArrayParam(Params, TEXT("vertices"), VerticesArray))
+	{
+		return FECACommandResult::Error(TEXT("Missing required parameter: vertices"));
+	}
+
+	const TArray<TSharedPtr<FJsonValue>>* TrianglesArray;
+	if (!GetArrayParam(Params, TEXT("triangles"), TrianglesArray))
+	{
+		return FECACommandResult::Error(TEXT("Missing required parameter: triangles"));
+	}
+
+	// Optional UV and normal arrays
+	const TArray<TSharedPtr<FJsonValue>>* UVsArray = nullptr;
+	GetArrayParam(Params, TEXT("uvs"), UVsArray, false);
+
+	const TArray<TSharedPtr<FJsonValue>>* NormalsArray = nullptr;
+	GetArrayParam(Params, TEXT("normals"), NormalsArray, false);
+
+	bool bGenerateCollision = true;
+	GetBoolParam(Params, TEXT("generate_collision"), bGenerateCollision, false);
+
+	bool bGenerateLightmapUVs = true;
+	GetBoolParam(Params, TEXT("generate_lightmap_uvs"), bGenerateLightmapUVs, false);
+
+	// Build the dynamic mesh
+	FDynamicMesh3 DynMesh;
+	DynMesh.EnableAttributes();
+	DynMesh.Attributes()->SetNumUVLayers(1);
+	DynMesh.Attributes()->EnableMaterialID();  // Enable material ID for multi-material support
+
+	// Parse material slots if provided
+	TArray<FString> MaterialSlotNames;
+	const TArray<TSharedPtr<FJsonValue>>* MaterialSlotsArray = nullptr;
+	if (GetArrayParam(Params, TEXT("material_slots"), MaterialSlotsArray, false) && MaterialSlotsArray)
+	{
+		for (const TSharedPtr<FJsonValue>& SlotValue : *MaterialSlotsArray)
+		{
+			FString SlotName;
+			if (SlotValue->TryGetString(SlotName))
+			{
+				MaterialSlotNames.Add(SlotName);
+			}
+			else
+			{
+				const TSharedPtr<FJsonObject>* SlotObj;
+				if (SlotValue->TryGetObject(SlotObj))
+				{
+					FString Name;
+					if ((*SlotObj)->TryGetStringField(TEXT("name"), Name))
+					{
+						MaterialSlotNames.Add(Name);
+					}
+					else
+					{
+						MaterialSlotNames.Add(FString::Printf(TEXT("Material_%d"), MaterialSlotNames.Num()));
+					}
+				}
+			}
+		}
+	}
+
+	// Parse vertex UVs if provided (per-vertex UVs)
+	TArray<FVector2f> VertexUVs;
+	if (UVsArray && UVsArray->Num() > 0)
+	{
+		for (const TSharedPtr<FJsonValue>& UVValue : *UVsArray)
+		{
+			const TSharedPtr<FJsonObject>* UVObj;
+			if (UVValue->TryGetObject(UVObj))
+			{
+				VertexUVs.Add(FVector2f(
+					(*UVObj)->GetNumberField(TEXT("u")),
+					(*UVObj)->GetNumberField(TEXT("v"))
+				));
+			}
+		}
+	}
+
+	// Parse vertex normals if provided
+	TArray<FVector3f> VertexNormals;
+	if (NormalsArray && NormalsArray->Num() > 0)
+	{
+		for (const TSharedPtr<FJsonValue>& NormalValue : *NormalsArray)
+		{
+			const TSharedPtr<FJsonObject>* NormalObj;
+			if (NormalValue->TryGetObject(NormalObj))
+			{
+				VertexNormals.Add(FVector3f(
+					(*NormalObj)->GetNumberField(TEXT("x")),
+					(*NormalObj)->GetNumberField(TEXT("y")),
+					(*NormalObj)->GetNumberField(TEXT("z"))
+				));
+			}
+		}
+	}
+
+	// Add vertices
+	TArray<int32> VertexMap;
+	VertexMap.Reserve(VerticesArray->Num());
+	
+	for (const TSharedPtr<FJsonValue>& VertexValue : *VerticesArray)
+	{
+		const TSharedPtr<FJsonObject>* VertexObj;
+		if (VertexValue->TryGetObject(VertexObj))
+		{
+			FVector3d Pos(
+				(*VertexObj)->GetNumberField(TEXT("x")),
+				(*VertexObj)->GetNumberField(TEXT("y")),
+				(*VertexObj)->GetNumberField(TEXT("z"))
+			);
+			int32 VertID = DynMesh.AppendVertex(Pos);
+			VertexMap.Add(VertID);
+		}
+	}
+
+	// Get UV and Normal overlays
+	FDynamicMeshUVOverlay* UVOverlay = DynMesh.Attributes()->PrimaryUV();
+	FDynamicMeshNormalOverlay* NormalOverlay = DynMesh.Attributes()->PrimaryNormals();
+
+	// Add triangles with UVs and normals
+	int32 TriIndex = 0;
+	for (const TSharedPtr<FJsonValue>& TriValue : *TrianglesArray)
+	{
+		const TSharedPtr<FJsonObject>* TriObj;
+		const TArray<TSharedPtr<FJsonValue>>* TriIndices;
+		
+		// Support both formats: array of indices [0,1,2] or object with indices and per-corner data
+		if (TriValue->TryGetObject(TriObj))
+		{
+			// Object format with optional per-corner UVs: {"indices": [0,1,2], "uvs": [{u,v}, {u,v}, {u,v}]}
+			const TArray<TSharedPtr<FJsonValue>>* IndicesArray;
+			if ((*TriObj)->TryGetArrayField(TEXT("indices"), IndicesArray) && IndicesArray->Num() >= 3)
+			{
+				int32 V0 = (*IndicesArray)[0]->AsNumber();
+				int32 V1 = (*IndicesArray)[1]->AsNumber();
+				int32 V2 = (*IndicesArray)[2]->AsNumber();
+
+				if (V0 >= 0 && V0 < VertexMap.Num() &&
+					V1 >= 0 && V1 < VertexMap.Num() &&
+					V2 >= 0 && V2 < VertexMap.Num())
+				{
+					// Reverse winding order (swap V1 and V2) for correct face orientation
+					int32 TriID = DynMesh.AppendTriangle(VertexMap[V0], VertexMap[V2], VertexMap[V1]);
+					
+					if (TriID >= 0)
+					{
+						// Per-triangle material index
+						int32 MaterialIndex = 0;
+						if ((*TriObj)->TryGetNumberField(TEXT("material_index"), MaterialIndex))
+						{
+							DynMesh.Attributes()->GetMaterialID()->SetValue(TriID, MaterialIndex);
+						}
+						
+						// Per-triangle UVs
+						const TArray<TSharedPtr<FJsonValue>>* TriUVsArray;
+						if ((*TriObj)->TryGetArrayField(TEXT("uvs"), TriUVsArray) && TriUVsArray->Num() >= 3)
+						{
+							FVector2f UV0, UV1, UV2;
+							const TSharedPtr<FJsonObject>* UV0Obj;
+							const TSharedPtr<FJsonObject>* UV1Obj;
+							const TSharedPtr<FJsonObject>* UV2Obj;
+							
+							if ((*TriUVsArray)[0]->TryGetObject(UV0Obj) &&
+								(*TriUVsArray)[1]->TryGetObject(UV1Obj) &&
+								(*TriUVsArray)[2]->TryGetObject(UV2Obj))
+							{
+								UV0 = FVector2f((*UV0Obj)->GetNumberField(TEXT("u")), (*UV0Obj)->GetNumberField(TEXT("v")));
+								UV1 = FVector2f((*UV1Obj)->GetNumberField(TEXT("u")), (*UV1Obj)->GetNumberField(TEXT("v")));
+								UV2 = FVector2f((*UV2Obj)->GetNumberField(TEXT("u")), (*UV2Obj)->GetNumberField(TEXT("v")));
+								
+								int32 UVID0 = UVOverlay->AppendElement(UV0);
+								int32 UVID1 = UVOverlay->AppendElement(UV1);
+								int32 UVID2 = UVOverlay->AppendElement(UV2);
+								// Match reversed winding order (swap UVID1 and UVID2)
+								UVOverlay->SetTriangle(TriID, FIndex3i(UVID0, UVID2, UVID1));
+							}
+						}
+						else if (VertexUVs.Num() > 0)
+						{
+							// Use per-vertex UVs if available
+							if (V0 < VertexUVs.Num() && V1 < VertexUVs.Num() && V2 < VertexUVs.Num())
+							{
+								int32 UVID0 = UVOverlay->AppendElement(VertexUVs[V0]);
+								int32 UVID1 = UVOverlay->AppendElement(VertexUVs[V1]);
+								int32 UVID2 = UVOverlay->AppendElement(VertexUVs[V2]);
+								// Match reversed winding order (swap UVID1 and UVID2)
+								UVOverlay->SetTriangle(TriID, FIndex3i(UVID0, UVID2, UVID1));
+							}
+						}
+					}
+				}
+			}
+		}
+		else if (TriValue->TryGetArray(TriIndices) && TriIndices->Num() >= 3)
+		{
+			// Simple array format [v0, v1, v2]
+			int32 V0 = (*TriIndices)[0]->AsNumber();
+			int32 V1 = (*TriIndices)[1]->AsNumber();
+			int32 V2 = (*TriIndices)[2]->AsNumber();
+
+			if (V0 >= 0 && V0 < VertexMap.Num() &&
+				V1 >= 0 && V1 < VertexMap.Num() &&
+				V2 >= 0 && V2 < VertexMap.Num())
+			{
+				// Reverse winding order (swap V1 and V2) for correct face orientation
+				int32 TriID = DynMesh.AppendTriangle(VertexMap[V0], VertexMap[V2], VertexMap[V1]);
+				
+				if (TriID >= 0 && VertexUVs.Num() > 0)
+				{
+					// Apply per-vertex UVs if available
+					if (V0 < VertexUVs.Num() && V1 < VertexUVs.Num() && V2 < VertexUVs.Num())
+					{
+						int32 UVID0 = UVOverlay->AppendElement(VertexUVs[V0]);
+						int32 UVID1 = UVOverlay->AppendElement(VertexUVs[V1]);
+						int32 UVID2 = UVOverlay->AppendElement(VertexUVs[V2]);
+						// Match reversed winding order (swap UVID1 and UVID2)
+						UVOverlay->SetTriangle(TriID, FIndex3i(UVID0, UVID2, UVID1));
+					}
+				}
+			}
+		}
+		TriIndex++;
+	}
+
+	// Compute normals if not provided
+	if (VertexNormals.Num() == 0)
+	{
+		FMeshNormals::QuickComputeVertexNormals(DynMesh);
+	}
+	else
+	{
+		// Apply provided normals
+		for (int32 TriID : DynMesh.TriangleIndicesItr())
+		{
+			FIndex3i Tri = DynMesh.GetTriangle(TriID);
+			
+			FVector3f N0 = (Tri.A < VertexNormals.Num()) ? VertexNormals[Tri.A] : FVector3f::UnitZ();
+			FVector3f N1 = (Tri.B < VertexNormals.Num()) ? VertexNormals[Tri.B] : FVector3f::UnitZ();
+			FVector3f N2 = (Tri.C < VertexNormals.Num()) ? VertexNormals[Tri.C] : FVector3f::UnitZ();
+			
+			int32 NID0 = NormalOverlay->AppendElement(N0);
+			int32 NID1 = NormalOverlay->AppendElement(N1);
+			int32 NID2 = NormalOverlay->AppendElement(N2);
+			NormalOverlay->SetTriangle(TriID, FIndex3i(NID0, NID1, NID2));
+		}
+	}
+
+	// Auto-generate planar UVs if none were provided
+	// Projects UVs based on the dominant axis of each triangle's normal
+	if (VertexUVs.Num() == 0 && (!UVsArray || UVsArray->Num() == 0))
+	{
+		// Get mesh bounds for UV scaling
+		FAxisAlignedBox3d Bounds = DynMesh.GetBounds();
+		double UVScale = 0.01; // 1 unit = 1cm, so 100 units = 1 UV tile
+		
+		for (int32 TriID : DynMesh.TriangleIndicesItr())
+		{
+			FIndex3i Tri = DynMesh.GetTriangle(TriID);
+			
+			FVector3d V0 = DynMesh.GetVertex(Tri.A);
+			FVector3d V1 = DynMesh.GetVertex(Tri.B);
+			FVector3d V2 = DynMesh.GetVertex(Tri.C);
+			
+			// Compute triangle normal to determine projection axis
+			FVector3d Edge1 = V1 - V0;
+			FVector3d Edge2 = V2 - V0;
+			FVector3d Normal = Edge1.Cross(Edge2);
+			Normal.Normalize();
+			
+			// Find dominant axis (the one most aligned with normal)
+			double AbsX = FMath::Abs(Normal.X);
+			double AbsY = FMath::Abs(Normal.Y);
+			double AbsZ = FMath::Abs(Normal.Z);
+			
+			FVector2f UV0, UV1, UV2;
+			
+			if (AbsX >= AbsY && AbsX >= AbsZ)
+			{
+				// X is dominant - project onto YZ plane
+				UV0 = FVector2f(V0.Y * UVScale, V0.Z * UVScale);
+				UV1 = FVector2f(V1.Y * UVScale, V1.Z * UVScale);
+				UV2 = FVector2f(V2.Y * UVScale, V2.Z * UVScale);
+			}
+			else if (AbsY >= AbsX && AbsY >= AbsZ)
+			{
+				// Y is dominant - project onto XZ plane
+				UV0 = FVector2f(V0.X * UVScale, V0.Z * UVScale);
+				UV1 = FVector2f(V1.X * UVScale, V1.Z * UVScale);
+				UV2 = FVector2f(V2.X * UVScale, V2.Z * UVScale);
+			}
+			else
+			{
+				// Z is dominant - project onto XY plane
+				UV0 = FVector2f(V0.X * UVScale, V0.Y * UVScale);
+				UV1 = FVector2f(V1.X * UVScale, V1.Y * UVScale);
+				UV2 = FVector2f(V2.X * UVScale, V2.Y * UVScale);
+			}
+			
+			int32 UVID0 = UVOverlay->AppendElement(UV0);
+			int32 UVID1 = UVOverlay->AppendElement(UV1);
+			int32 UVID2 = UVOverlay->AppendElement(UV2);
+			UVOverlay->SetTriangle(TriID, FIndex3i(UVID0, UVID1, UVID2));
+		}
+	}
+
+	// Create the static mesh asset
+	FString Error;
+	UStaticMesh* StaticMesh = MeshHelpers::CreateStaticMeshFromDynamicMesh(DynMesh, AssetPath, bGenerateCollision, bGenerateLightmapUVs, &Error);
+	
+	if (!StaticMesh)
+	{
+		return FECACommandResult::Error(Error);
+	}
+
+	// Apply custom material slot names if provided
+	if (MaterialSlotNames.Num() > 0)
+	{
+		TArray<FStaticMaterial>& Materials = StaticMesh->GetStaticMaterials();
+		for (int32 i = 0; i < MaterialSlotNames.Num() && i < Materials.Num(); i++)
+		{
+			Materials[i].MaterialSlotName = FName(*MaterialSlotNames[i]);
+			Materials[i].ImportedMaterialSlotName = FName(*MaterialSlotNames[i]);
+		}
+		// Add any additional slots specified beyond what the mesh uses
+		for (int32 i = Materials.Num(); i < MaterialSlotNames.Num(); i++)
+		{
+			FStaticMaterial NewMaterial(nullptr, FName(*MaterialSlotNames[i]), FName(*MaterialSlotNames[i]));
+			NewMaterial.UVChannelData = FMeshUVChannelInfo(1.0f);
+			StaticMesh->GetStaticMaterials().Add(NewMaterial);
+		}
+	}
+
+	TSharedPtr<FJsonObject> Result = MeshHelpers::BuildMeshResultJson(StaticMesh, DynMesh.VertexCount(), DynMesh.TriangleCount());
+	return FECACommandResult::Success(Result);
+}
+
+//==============================================================================
+// CREATE MESH FROM FACES - High-level face-based mesh creation for LLMs
+//==============================================================================
+
+// Helper to generate planar UV based on dominant normal axis
+static FVector2f ProjectUVByNormal(const FVector3d& Pos, const FVector3d& Normal, double UVScale)
+{
+	double AbsX = FMath::Abs(Normal.X);
+	double AbsY = FMath::Abs(Normal.Y);
+	double AbsZ = FMath::Abs(Normal.Z);
+	
+	if (AbsX >= AbsY && AbsX >= AbsZ)
+		return FVector2f(Pos.Y * UVScale, Pos.Z * UVScale);
+	else if (AbsY >= AbsX && AbsY >= AbsZ)
+		return FVector2f(Pos.X * UVScale, Pos.Z * UVScale);
+	else
+		return FVector2f(Pos.X * UVScale, Pos.Y * UVScale);
+}
+
+FECACommandResult FECACommand_CreateMeshFromFaces::Execute(const TSharedPtr<FJsonObject>& Params)
+{
+	FString AssetPath;
+	if (!GetStringParam(Params, TEXT("asset_path"), AssetPath))
+	{
+		return FECACommandResult::Error(TEXT("Missing required parameter: asset_path"));
+	}
+
+	const TArray<TSharedPtr<FJsonValue>>* FacesArray;
+	if (!GetArrayParam(Params, TEXT("faces"), FacesArray))
+	{
+		return FECACommandResult::Error(TEXT("Missing required parameter: faces"));
+	}
+
+	// Optional parameters
+	double Thickness = 0;
+	GetFloatParam(Params, TEXT("thickness"), Thickness, false);
+	
+	double UVScale = 0.01; // Default: 100 units = 1 UV tile
+	GetFloatParam(Params, TEXT("uv_scale"), UVScale, false);
+
+	bool bGenerateCollision = true;
+	GetBoolParam(Params, TEXT("generate_collision"), bGenerateCollision, false);
+
+	bool bDoubleSided = false;
+	GetBoolParam(Params, TEXT("double_sided"), bDoubleSided, false);
+
+	// Create dynamic mesh
+	FDynamicMesh3 DynMesh;
+	DynMesh.EnableAttributes();
+	DynMesh.Attributes()->SetNumUVLayers(1);
+	DynMesh.Attributes()->EnableMaterialID();
+
+	FDynamicMeshUVOverlay* UVOverlay = DynMesh.Attributes()->PrimaryUV();
+
+	// Track face names for result
+	TArray<FString> FaceNames;
+	int32 MaterialIndex = 0;
+
+	// Process each face
+	for (const TSharedPtr<FJsonValue>& FaceValue : *FacesArray)
+	{
+		const TSharedPtr<FJsonObject>* FaceObj;
+		if (!FaceValue->TryGetObject(FaceObj))
+		{
+			continue;
+		}
+
+		// Get face name (optional)
+		FString FaceName;
+		if (!(*FaceObj)->TryGetStringField(TEXT("name"), FaceName))
+		{
+			FaceName = FString::Printf(TEXT("Face_%d"), FaceNames.Num());
+		}
+		FaceNames.Add(FaceName);
+
+		// Get face material index (optional, defaults to face order)
+		int32 FaceMaterialIndex = MaterialIndex;
+		(*FaceObj)->TryGetNumberField(TEXT("material_index"), FaceMaterialIndex);
+
+		// Get vertices for this face
+		const TArray<TSharedPtr<FJsonValue>>* VerticesArray;
+		if (!(*FaceObj)->TryGetArrayField(TEXT("vertices"), VerticesArray) || VerticesArray->Num() < 3)
+		{
+			continue;
+		}
+
+		// Parse vertices - support both array format [x,y,z] and object format {x,y,z}
+		TArray<FVector3d> FaceVertices;
+		for (const TSharedPtr<FJsonValue>& VertValue : *VerticesArray)
+		{
+			const TSharedPtr<FJsonObject>* VertObj;
+			const TArray<TSharedPtr<FJsonValue>>* VertArr;
+			
+			if (VertValue->TryGetObject(VertObj))
+			{
+				FaceVertices.Add(FVector3d(
+					(*VertObj)->GetNumberField(TEXT("x")),
+					(*VertObj)->GetNumberField(TEXT("y")),
+					(*VertObj)->GetNumberField(TEXT("z"))
+				));
+			}
+			else if (VertValue->TryGetArray(VertArr) && VertArr->Num() >= 3)
+			{
+				FaceVertices.Add(FVector3d(
+					(*VertArr)[0]->AsNumber(),
+					(*VertArr)[1]->AsNumber(),
+					(*VertArr)[2]->AsNumber()
+				));
+			}
+		}
+
+		if (FaceVertices.Num() < 3)
+		{
+			continue;
+		}
+
+		// Compute face normal from vertex winding
+		// Note: Edge2.Cross(Edge1) gives correct normal direction for UE's coordinate system
+		FVector3d Edge1 = FaceVertices[1] - FaceVertices[0];
+		FVector3d Edge2 = FaceVertices[FaceVertices.Num() - 1] - FaceVertices[0];
+		FVector3d FaceNormal = Edge2.Cross(Edge1);
+		FaceNormal.Normalize();
+
+		// Check for normal_hint - if provided, ensure face normal aligns with it
+		// If the computed normal points opposite to the hint, we'll flip the winding
+		bool bFlipWinding = false;
+		const TSharedPtr<FJsonObject>* NormalHintObj;
+		const TArray<TSharedPtr<FJsonValue>>* NormalHintArr;
+		FVector3d NormalHint = FVector3d::Zero();
+		
+		if ((*FaceObj)->TryGetObjectField(TEXT("normal_hint"), NormalHintObj))
+		{
+			NormalHint = FVector3d(
+				(*NormalHintObj)->GetNumberField(TEXT("x")),
+				(*NormalHintObj)->GetNumberField(TEXT("y")),
+				(*NormalHintObj)->GetNumberField(TEXT("z"))
+			);
+		}
+		else if ((*FaceObj)->TryGetArrayField(TEXT("normal_hint"), NormalHintArr) && NormalHintArr->Num() >= 3)
+		{
+			NormalHint = FVector3d(
+				(*NormalHintArr)[0]->AsNumber(),
+				(*NormalHintArr)[1]->AsNumber(),
+				(*NormalHintArr)[2]->AsNumber()
+			);
+		}
+		
+		// If normal_hint provided, use it to correct winding
+		if (!NormalHint.IsZero())
+		{
+			NormalHint.Normalize();
+			// If face normal points opposite to hint (dot product negative), flip winding
+			if (FaceNormal.Dot(NormalHint) < 0)
+			{
+				bFlipWinding = true;
+				FaceNormal = -FaceNormal;
+			}
+		}
+		// No auto-detection - user must specify normal_hint for explicit control,
+		// or ensure vertices are in correct winding order (CCW when viewed from front)
+
+		// Add vertices to mesh
+		TArray<int32> VertexIDs;
+		TArray<int32> UVIDs;
+		
+		for (const FVector3d& Vert : FaceVertices)
+		{
+			int32 VID = DynMesh.AppendVertex(Vert);
+			VertexIDs.Add(VID);
+
+			// Generate planar UV based on dominant axis
+			double AbsX = FMath::Abs(FaceNormal.X);
+			double AbsY = FMath::Abs(FaceNormal.Y);
+			double AbsZ = FMath::Abs(FaceNormal.Z);
+			
+			FVector2f UV;
+			if (AbsX >= AbsY && AbsX >= AbsZ)
+			{
+				UV = FVector2f(Vert.Y * UVScale, Vert.Z * UVScale);
+			}
+			else if (AbsY >= AbsX && AbsY >= AbsZ)
+			{
+				UV = FVector2f(Vert.X * UVScale, Vert.Z * UVScale);
+			}
+			else
+			{
+				UV = FVector2f(Vert.X * UVScale, Vert.Y * UVScale);
+			}
+			
+			int32 UVID = UVOverlay->AppendElement(UV);
+			UVIDs.Add(UVID);
+		}
+
+		// Use proper polygon triangulation (handles non-convex polygons)
+		TArray<FIndex3i> Triangles;
+		PolygonTriangulation::TriangulateSimplePolygon<double>(FaceVertices, Triangles, false);
+
+		for (const FIndex3i& Tri : Triangles)
+		{
+			int32 A = Tri.A;
+			int32 B = Tri.B;
+			int32 C = Tri.C;
+
+			// Ensure indices are valid
+			if (A < 0 || A >= VertexIDs.Num() || 
+				B < 0 || B >= VertexIDs.Num() || 
+				C < 0 || C >= VertexIDs.Num())
+			{
+				continue;
+			}
+
+			// Create front face triangle
+			// Apply winding based on normal_hint correction
+			int32 TriID;
+			if (bFlipWinding)
+			{
+				TriID = DynMesh.AppendTriangle(VertexIDs[A], VertexIDs[B], VertexIDs[C]);
+			}
+			else
+			{
+				TriID = DynMesh.AppendTriangle(VertexIDs[A], VertexIDs[C], VertexIDs[B]);
+			}
+			
+			if (TriID >= 0)
+			{
+				DynMesh.Attributes()->GetMaterialID()->SetValue(TriID, FaceMaterialIndex);
+				if (bFlipWinding)
+					UVOverlay->SetTriangle(TriID, FIndex3i(UVIDs[A], UVIDs[B], UVIDs[C]));
+				else
+					UVOverlay->SetTriangle(TriID, FIndex3i(UVIDs[A], UVIDs[C], UVIDs[B]));
+			}
+
+			// Add back face if double-sided (opposite winding)
+			if (bDoubleSided)
+			{
+				int32 BackTriID;
+				if (bFlipWinding)
+				{
+					BackTriID = DynMesh.AppendTriangle(VertexIDs[A], VertexIDs[C], VertexIDs[B]);
+				}
+				else
+				{
+					BackTriID = DynMesh.AppendTriangle(VertexIDs[A], VertexIDs[B], VertexIDs[C]);
+				}
+				
+				if (BackTriID >= 0)
+				{
+					DynMesh.Attributes()->GetMaterialID()->SetValue(BackTriID, FaceMaterialIndex);
+					if (bFlipWinding)
+						UVOverlay->SetTriangle(BackTriID, FIndex3i(UVIDs[A], UVIDs[C], UVIDs[B]));
+					else
+						UVOverlay->SetTriangle(BackTriID, FIndex3i(UVIDs[A], UVIDs[B], UVIDs[C]));
+				}
+			}
+		}
+
+		MaterialIndex++;
+	}
+
+	// If thickness specified, extrude the mesh to create a solid
+	if (Thickness > 0)
+	{
+		// Store original triangle count before we add more
+		TArray<FIndex3i> OriginalTris;
+		for (int32 TriID : DynMesh.TriangleIndicesItr())
+		{
+			OriginalTris.Add(DynMesh.GetTriangle(TriID));
+		}
+		
+		// Get boundary edges before adding new geometry
+		// Store edges with correct winding (matching the triangle they belong to)
+		TArray<FIndex2i> BoundaryEdges;
+		for (int32 EdgeID : DynMesh.EdgeIndicesItr())
+		{
+			if (DynMesh.IsBoundaryEdge(EdgeID))
+			{
+				FIndex2i EdgeV = DynMesh.GetEdgeV(EdgeID);
+				
+				// Get the triangle this edge belongs to and check winding
+				FIndex2i EdgeT = DynMesh.GetEdgeT(EdgeID);
+				int32 TriID = (EdgeT.A != -1) ? EdgeT.A : EdgeT.B;
+				
+				if (TriID != -1)
+				{
+					FIndex3i Tri = DynMesh.GetTriangle(TriID);
+					// Find the edge in the triangle and check its direction
+					// Triangle edges are: (A,B), (B,C), (C,A)
+					// We want edge direction that matches CCW winding of the triangle
+					if ((Tri.A == EdgeV.A && Tri.B == EdgeV.B) ||
+						(Tri.B == EdgeV.A && Tri.C == EdgeV.B) ||
+						(Tri.C == EdgeV.A && Tri.A == EdgeV.B))
+					{
+						// Edge is in correct order relative to triangle winding
+						BoundaryEdges.Add(EdgeV);
+					}
+					else
+					{
+						// Edge is reversed, flip it
+						BoundaryEdges.Add(FIndex2i(EdgeV.B, EdgeV.A));
+					}
+				}
+				else
+				{
+					BoundaryEdges.Add(EdgeV);
+				}
+			}
+		}
+		
+		// Create offset vertices (new vertices for the back face)
+		TMap<int32, int32> VertexToOffsetVertex;
+		int32 OriginalVertCount = DynMesh.MaxVertexID();
+		
+		// Compute average normal per vertex for offset direction
+		TMap<int32, FVector3d> VertexNormals;
+		for (const FIndex3i& Tri : OriginalTris)
+		{
+			FVector3d V0 = DynMesh.GetVertex(Tri.A);
+			FVector3d V1 = DynMesh.GetVertex(Tri.B);
+			FVector3d V2 = DynMesh.GetVertex(Tri.C);
+			FVector3d TriNormal = (V1 - V0).Cross(V2 - V0);
+			TriNormal.Normalize();
+			
+			VertexNormals.FindOrAdd(Tri.A) += TriNormal;
+			VertexNormals.FindOrAdd(Tri.B) += TriNormal;
+			VertexNormals.FindOrAdd(Tri.C) += TriNormal;
+		}
+		
+		// Create offset vertices
+		for (auto& Pair : VertexNormals)
+		{
+			int32 VID = Pair.Key;
+			FVector3d Normal = Pair.Value;
+			Normal.Normalize();
+			
+			FVector3d Pos = DynMesh.GetVertex(VID);
+			FVector3d OffsetPos = Pos - Normal * Thickness; // Offset inward (opposite to normal)
+			
+			int32 NewVID = DynMesh.AppendVertex(OffsetPos);
+			VertexToOffsetVertex.Add(VID, NewVID);
+		}
+		
+		// Create back faces with reversed winding
+		for (const FIndex3i& OrigTri : OriginalTris)
+		{
+			int32 V0_Off = VertexToOffsetVertex[OrigTri.A];
+			int32 V1_Off = VertexToOffsetVertex[OrigTri.B];
+			int32 V2_Off = VertexToOffsetVertex[OrigTri.C];
+			
+			// Back face needs opposite winding from front face
+			// Front was A,B,C or A,C,B depending on bFlipWinding, so back is the opposite
+			int32 BackTriID = DynMesh.AppendTriangle(V0_Off, V1_Off, V2_Off);
+			if (BackTriID >= 0)
+			{
+				FVector3d P0 = DynMesh.GetVertex(V0_Off);
+				FVector3d P1 = DynMesh.GetVertex(V2_Off);
+				FVector3d P2 = DynMesh.GetVertex(V1_Off);
+				
+				FVector3d BackNormal = (P1 - P0).Cross(P2 - P0);
+				BackNormal.Normalize();
+				
+				int32 UVID0 = UVOverlay->AppendElement(ProjectUVByNormal(P0, BackNormal, UVScale));
+				int32 UVID1 = UVOverlay->AppendElement(ProjectUVByNormal(P1, BackNormal, UVScale));
+				int32 UVID2 = UVOverlay->AppendElement(ProjectUVByNormal(P2, BackNormal, UVScale));
+				UVOverlay->SetTriangle(BackTriID, FIndex3i(UVID0, UVID1, UVID2));
+			}
+		}
+		
+		// Create side faces from boundary edges
+		for (const FIndex2i& EdgeV : BoundaryEdges)
+		{
+			int32 V0 = EdgeV.A;
+			int32 V1 = EdgeV.B;
+			
+			if (!VertexToOffsetVertex.Contains(V0) || !VertexToOffsetVertex.Contains(V1))
+				continue;
+				
+			int32 V0_Off = VertexToOffsetVertex[V0];
+			int32 V1_Off = VertexToOffsetVertex[V1];
+			
+			FVector3d P0 = DynMesh.GetVertex(V0);
+			FVector3d P1 = DynMesh.GetVertex(V1);
+			FVector3d P0_Off = DynMesh.GetVertex(V0_Off);
+			FVector3d P1_Off = DynMesh.GetVertex(V1_Off);
+			
+			// Compute side face normal for UV projection
+			FVector3d SideNormal = (P1 - P0).Cross(P0_Off - P0);
+			SideNormal.Normalize();
+			
+			// Create two triangles for the quad side face
+			// Winding: V0 -> V1 -> V1_Off and V0 -> V1_Off -> V0_Off
+			int32 Tri1 = DynMesh.AppendTriangle(V0, V1, V1_Off);
+			int32 Tri2 = DynMesh.AppendTriangle(V0, V1_Off, V0_Off);
+			
+			if (Tri1 >= 0)
+			{
+				int32 UV0 = UVOverlay->AppendElement(ProjectUVByNormal(P0, SideNormal, UVScale));
+				int32 UV1 = UVOverlay->AppendElement(ProjectUVByNormal(P1, SideNormal, UVScale));
+				int32 UV2 = UVOverlay->AppendElement(ProjectUVByNormal(P1_Off, SideNormal, UVScale));
+				UVOverlay->SetTriangle(Tri1, FIndex3i(UV0, UV1, UV2));
+			}
+			if (Tri2 >= 0)
+			{
+				int32 UV0 = UVOverlay->AppendElement(ProjectUVByNormal(P0, SideNormal, UVScale));
+				int32 UV1 = UVOverlay->AppendElement(ProjectUVByNormal(P1_Off, SideNormal, UVScale));
+				int32 UV2 = UVOverlay->AppendElement(ProjectUVByNormal(P0_Off, SideNormal, UVScale));
+				UVOverlay->SetTriangle(Tri2, FIndex3i(UV0, UV1, UV2));
+			}
+		}
+	}
+
+	// Compute normals
+	FMeshNormals::QuickComputeVertexNormals(DynMesh);
+
+	// Create static mesh
+	FString Error;
+	UStaticMesh* StaticMesh = MeshHelpers::CreateStaticMeshFromDynamicMesh(DynMesh, AssetPath, bGenerateCollision, true, &Error);
+	
+	if (!StaticMesh)
+	{
+		return FECACommandResult::Error(Error);
+	}
+
+	// Build result
+	TSharedPtr<FJsonObject> Result = MeshHelpers::BuildMeshResultJson(StaticMesh, DynMesh.VertexCount(), DynMesh.TriangleCount());
+	
+	// Add face names to result
+	TArray<TSharedPtr<FJsonValue>> FaceNamesJson;
+	for (const FString& Name : FaceNames)
+	{
+		FaceNamesJson.Add(MakeShared<FJsonValueString>(Name));
+	}
+	Result->SetArrayField(TEXT("face_names"), FaceNamesJson);
+	
+	return FECACommandResult::Success(Result);
+}
+
+//==============================================================================
+// POLYGON MESH
+//==============================================================================
+
+FECACommandResult FECACommand_CreatePolygonMesh::Execute(const TSharedPtr<FJsonObject>& Params)
+{
+	FString AssetPath;
+	if (!GetStringParam(Params, TEXT("asset_path"), AssetPath))
+	{
+		return FECACommandResult::Error(TEXT("Missing required parameter: asset_path"));
+	}
+
+	const TArray<TSharedPtr<FJsonValue>>* PointsArray;
+	if (!GetArrayParam(Params, TEXT("points"), PointsArray))
+	{
+		return FECACommandResult::Error(TEXT("Missing required parameter: points"));
+	}
+
+	double Height = 0;
+	GetFloatParam(Params, TEXT("height"), Height, false);
+
+	// Parse polygon points
+	TArray<FVector2d> Points;
+	for (const TSharedPtr<FJsonValue>& PointValue : *PointsArray)
+	{
+		const TSharedPtr<FJsonObject>* PointObj;
+		if (PointValue->TryGetObject(PointObj))
+		{
+			Points.Add(FVector2d(
+				(*PointObj)->GetNumberField(TEXT("x")),
+				(*PointObj)->GetNumberField(TEXT("y"))
+			));
+		}
+	}
+
+	if (Points.Num() < 3)
+	{
+		return FECACommandResult::Error(TEXT("Polygon must have at least 3 points"));
+	}
+
+	// Create dynamic mesh
+	FDynamicMesh3 DynMesh;
+	DynMesh.EnableAttributes();
+	DynMesh.Attributes()->SetNumUVLayers(1);
+
+	if (Height > 0)
+	{
+		// Extrude polygon using GeneralizedCylinderGenerator
+		FPolygon2d Profile(Points);
+		
+		TArray<FVector3d> Path;
+		Path.Add(FVector3d(0, 0, 0));
+		Path.Add(FVector3d(0, 0, Height));
+
+		FGeneralizedCylinderGenerator ExtrudeGen;
+		ExtrudeGen.CrossSection = Profile;
+		ExtrudeGen.Path = Path;
+		ExtrudeGen.bCapped = true;
+		ExtrudeGen.bPolygroupPerQuad = false;
+	
+		ExtrudeGen.Generate();
+		DynMesh.Copy(&ExtrudeGen);
+	}
+	else
+	{
+		// Flat polygon - simple triangulation using ear clipping
+		// For now, create a simple fan triangulation (works for convex polygons)
+		TArray<int32> VertexIDs;
+		for (const FVector2d& Point : Points)
+		{
+			VertexIDs.Add(DynMesh.AppendVertex(FVector3d(Point.X, Point.Y, 0)));
+		}
+
+		// Fan triangulation from first vertex (reversed winding for correct orientation)
+		for (int32 i = 1; i < VertexIDs.Num() - 1; i++)
+		{
+			DynMesh.AppendTriangle(VertexIDs[0], VertexIDs[i + 1], VertexIDs[i]);
+		}
+	}
+
+	FMeshNormals::QuickComputeVertexNormals(DynMesh);
+
+	// Auto-generate planar UVs (polygon mesh is 2D, so project onto XY plane)
+	FDynamicMeshUVOverlay* UVOverlay = DynMesh.Attributes()->PrimaryUV();
+	double UVScale = 0.01; // 1 unit = 1cm, so 100 units = 1 UV tile
+	
+	for (int32 TriID : DynMesh.TriangleIndicesItr())
+	{
+		FIndex3i Tri = DynMesh.GetTriangle(TriID);
+		
+		FVector3d V0 = DynMesh.GetVertex(Tri.A);
+		FVector3d V1 = DynMesh.GetVertex(Tri.B);
+		FVector3d V2 = DynMesh.GetVertex(Tri.C);
+		
+		// Project onto XY plane (polygon mesh is 2D in XY)
+		FVector2f UV0(V0.X * UVScale, V0.Y * UVScale);
+		FVector2f UV1(V1.X * UVScale, V1.Y * UVScale);
+		FVector2f UV2(V2.X * UVScale, V2.Y * UVScale);
+		
+		int32 UVID0 = UVOverlay->AppendElement(UV0);
+		int32 UVID1 = UVOverlay->AppendElement(UV1);
+		int32 UVID2 = UVOverlay->AppendElement(UV2);
+		UVOverlay->SetTriangle(TriID, FIndex3i(UVID0, UVID1, UVID2));
+	}
+
+	// Create static mesh
+	FString Error;
+	UStaticMesh* StaticMesh = MeshHelpers::CreateStaticMeshFromDynamicMesh(DynMesh, AssetPath, true, true, &Error);
+	
+	if (!StaticMesh)
+	{
+		return FECACommandResult::Error(Error);
+	}
+
+	// Apply material if specified
+	FString MaterialPath;
+	if (GetStringParam(Params, TEXT("material_path"), MaterialPath, false))
+	{
+		UMaterialInterface* Material = LoadObject<UMaterialInterface>(nullptr, *MaterialPath);
+		if (Material)
+		{
+			StaticMesh->SetMaterial(0, Material);
+		}
+	}
+
+	TSharedPtr<FJsonObject> Result = MeshHelpers::BuildMeshResultJson(StaticMesh, DynMesh.VertexCount(), DynMesh.TriangleCount());
+	return FECACommandResult::Success(Result);
+}
+
+//==============================================================================
+// MESH BOOLEAN - Stub (requires engine's mesh boolean which may not be exposed)
+//==============================================================================
+
+FECACommandResult FECACommand_MeshBoolean::Execute(const TSharedPtr<FJsonObject>& Params)
+{
+	FString OutputPath;
+	if (!GetStringParam(Params, TEXT("output_path"), OutputPath))
+	{
+		return FECACommandResult::Error(TEXT("Missing required parameter: output_path"));
+	}
+
+	FString MeshAPath;
+	if (!GetStringParam(Params, TEXT("mesh_a"), MeshAPath))
+	{
+		return FECACommandResult::Error(TEXT("Missing required parameter: mesh_a"));
+	}
+
+	FString MeshBPath;
+	if (!GetStringParam(Params, TEXT("mesh_b"), MeshBPath))
+	{
+		return FECACommandResult::Error(TEXT("Missing required parameter: mesh_b"));
+	}
+
+	FString OperationStr = TEXT("difference");
+	GetStringParam(Params, TEXT("operation"), OperationStr, false);
+
+	// Load mesh A
+	UStaticMesh* MeshA = MeshHelpers::LoadStaticMesh(MeshAPath);
+	if (!MeshA)
+	{
+		return FECACommandResult::Error(FString::Printf(TEXT("Failed to load mesh_a: %s"), *MeshAPath));
+	}
+
+	// Load mesh B
+	UStaticMesh* MeshB = MeshHelpers::LoadStaticMesh(MeshBPath);
+	if (!MeshB)
+	{
+		return FECACommandResult::Error(FString::Printf(TEXT("Failed to load mesh_b: %s"), *MeshBPath));
+	}
+
+	// Convert to dynamic meshes
+	FDynamicMesh3 DynMeshA;
+	FDynamicMesh3 DynMeshB;
+
+	FMeshDescriptionToDynamicMesh ConverterA;
+	ConverterA.Convert(MeshA->GetMeshDescription(0), DynMeshA);
+
+	FMeshDescriptionToDynamicMesh ConverterB;
+	ConverterB.Convert(MeshB->GetMeshDescription(0), DynMeshB);
+
+	// Parse optional transform for mesh B
+	FTransformSRT3d TransformA = FTransformSRT3d::Identity();
+	FTransformSRT3d TransformB = FTransformSRT3d::Identity();
+
+	const TSharedPtr<FJsonObject>* TransformBObj;
+	if (Params->TryGetObjectField(TEXT("transform_b"), TransformBObj))
+	{
+		const TSharedPtr<FJsonObject>* LocObj;
+		if ((*TransformBObj)->TryGetObjectField(TEXT("location"), LocObj))
+		{
+			FVector3d Location(
+				(*LocObj)->GetNumberField(TEXT("x")),
+				(*LocObj)->GetNumberField(TEXT("y")),
+				(*LocObj)->GetNumberField(TEXT("z"))
+			);
+			TransformB.SetTranslation(Location);
+		}
+
+		const TSharedPtr<FJsonObject>* RotObj;
+		if ((*TransformBObj)->TryGetObjectField(TEXT("rotation"), RotObj))
+		{
+			FRotator Rotation(
+				(*RotObj)->GetNumberField(TEXT("pitch")),
+				(*RotObj)->GetNumberField(TEXT("yaw")),
+				(*RotObj)->GetNumberField(TEXT("roll"))
+			);
+			TransformB.SetRotation(FQuaterniond(Rotation.Quaternion()));
+		}
+
+		const TSharedPtr<FJsonObject>* ScaleObj;
+		if ((*TransformBObj)->TryGetObjectField(TEXT("scale"), ScaleObj))
+		{
+			FVector3d Scale(
+				(*ScaleObj)->GetNumberField(TEXT("x")),
+				(*ScaleObj)->GetNumberField(TEXT("y")),
+				(*ScaleObj)->GetNumberField(TEXT("z"))
+			);
+			TransformB.SetScale(Scale);
+		}
+	}
+
+	// Determine operation type
+	FMeshBoolean::EBooleanOp BoolOp = FMeshBoolean::EBooleanOp::Difference;
+	if (OperationStr.Equals(TEXT("union"), ESearchCase::IgnoreCase))
+	{
+		BoolOp = FMeshBoolean::EBooleanOp::Union;
+	}
+	else if (OperationStr.Equals(TEXT("difference"), ESearchCase::IgnoreCase))
+	{
+		BoolOp = FMeshBoolean::EBooleanOp::Difference;
+	}
+	else if (OperationStr.Equals(TEXT("intersect"), ESearchCase::IgnoreCase) || 
+	         OperationStr.Equals(TEXT("intersection"), ESearchCase::IgnoreCase))
+	{
+		BoolOp = FMeshBoolean::EBooleanOp::Intersect;
+	}
+	else if (OperationStr.Equals(TEXT("trim_inside"), ESearchCase::IgnoreCase))
+	{
+		BoolOp = FMeshBoolean::EBooleanOp::TrimInside;
+	}
+	else if (OperationStr.Equals(TEXT("trim_outside"), ESearchCase::IgnoreCase))
+	{
+		BoolOp = FMeshBoolean::EBooleanOp::TrimOutside;
+	}
+
+	// Set up boolean operation
+	FDynamicMesh3 ResultMesh;
+	FMeshBoolean MeshBoolean(
+		&DynMeshA, TransformA,
+		&DynMeshB, TransformB,
+		&ResultMesh, BoolOp
+	);
+	MeshBoolean.bPutResultInInputSpace = true;
+	MeshBoolean.bWeldSharedEdges = true;
+
+	// Perform boolean
+	if (!MeshBoolean.Compute())
+	{
+		return FECACommandResult::Error(TEXT("Boolean operation failed"));
+	}
+
+	// Check result
+	if (ResultMesh.TriangleCount() == 0)
+	{
+		return FECACommandResult::Error(TEXT("Boolean operation produced empty result"));
+	}
+
+	// Compute normals
+	FMeshNormals::QuickComputeVertexNormals(ResultMesh);
+
+	// Create static mesh from result
+	FString Error;
+	UStaticMesh* OutputMesh = MeshHelpers::CreateStaticMeshFromDynamicMesh(ResultMesh, OutputPath, true, true, &Error);
+	if (!OutputMesh)
+	{
+		return FECACommandResult::Error(Error);
+	}
+
+	// Build result
+	TSharedPtr<FJsonObject> Result = MeshHelpers::BuildMeshResultJson(OutputMesh, ResultMesh.VertexCount(), ResultMesh.TriangleCount());
+	Result->SetStringField(TEXT("operation"), OperationStr);
+	return FECACommandResult::Success(Result);
+}
+
+FECACommandResult FECACommand_MeshSelfUnion::Execute(const TSharedPtr<FJsonObject>& Params)
+{
+	return FECACommandResult::Error(TEXT("mesh_self_union: Requires experimental boolean operations. Not yet available."));
+}
+
+//==============================================================================
+// MESH PLANE CUT
+//==============================================================================
+
+FECACommandResult FECACommand_MeshPlaneCut::Execute(const TSharedPtr<FJsonObject>& Params)
+{
+	FString MeshPath;
+	if (!GetStringParam(Params, TEXT("mesh_path"), MeshPath))
+	{
+		return FECACommandResult::Error(TEXT("Missing required parameter: mesh_path"));
+	}
+
+	const TSharedPtr<FJsonObject>* PlaneOriginObj;
+	if (!GetObjectParam(Params, TEXT("plane_origin"), PlaneOriginObj))
+	{
+		return FECACommandResult::Error(TEXT("Missing required parameter: plane_origin"));
+	}
+
+	const TSharedPtr<FJsonObject>* PlaneNormalObj;
+	if (!GetObjectParam(Params, TEXT("plane_normal"), PlaneNormalObj))
+	{
+		return FECACommandResult::Error(TEXT("Missing required parameter: plane_normal"));
+	}
+
+	FString OutputPath;
+	if (!GetStringParam(Params, TEXT("output_path"), OutputPath, false))
+	{
+		OutputPath = MeshPath;
+	}
+
+	bool bFillHole = true;
+	GetBoolParam(Params, TEXT("fill_hole"), bFillHole, false);
+
+	FVector3d PlaneOrigin = MeshHelpers::ParseVector(*PlaneOriginObj);
+	FVector3d PlaneNormal = MeshHelpers::ParseVector(*PlaneNormalObj);
+	PlaneNormal = PlaneNormal.IsNearlyZero() ? FVector3d::UnitZ() : PlaneNormal.GetSafeNormal();
+
+	UStaticMesh* SourceMesh = MeshHelpers::LoadStaticMesh(MeshPath);
+	if (!SourceMesh)
+	{
+		return FECACommandResult::Error(FString::Printf(TEXT("Mesh not found: %s"), *MeshPath));
+	}
+
+	FDynamicMesh3 DynMesh;
+	FString Error;
+	if (!MeshHelpers::LoadStaticMeshIntoDynamicMesh(SourceMesh, DynMesh, &Error))
+	{
+		return FECACommandResult::Error(Error);
+	}
+
+	// Perform plane cut
+	FMeshPlaneCut PlaneCut(&DynMesh, PlaneOrigin, PlaneNormal);
+	PlaneCut.Cut();
+
+	// Fill holes if requested
+	if (bFillHole && PlaneCut.OpenBoundaries.Num() > 0)
+	{
+		for (const FMeshPlaneCut::FOpenBoundary& Boundary : PlaneCut.OpenBoundaries)
+		{
+			for (const FEdgeLoop& Loop : Boundary.CutLoops)
+			{
+				FSimpleHoleFiller Filler(&DynMesh, Loop);
+				Filler.Fill();
+			}
+		}
+	}
+
+	FMeshNormals::QuickComputeVertexNormals(DynMesh);
+
+	// Save result
+	UStaticMesh* OutputMesh = nullptr;
+	if (OutputPath == MeshPath)
+	{
+		MeshHelpers::SaveDynamicMeshToStaticMesh(DynMesh, SourceMesh);
+		OutputMesh = SourceMesh;
+	}
+	else
+	{
+		OutputMesh = MeshHelpers::CreateStaticMeshFromDynamicMesh(DynMesh, OutputPath, true, true, &Error);
+		if (!OutputMesh)
+		{
+			return FECACommandResult::Error(Error);
+		}
+	}
+
+	TSharedPtr<FJsonObject> Result = MeshHelpers::BuildMeshResultJson(OutputMesh);
+	return FECACommandResult::Success(Result);
+}
+
+//==============================================================================
+// MESH SLICE
+//==============================================================================
+
+FECACommandResult FECACommand_MeshSlice::Execute(const TSharedPtr<FJsonObject>& Params)
+{
+	FString MeshPath;
+	if (!GetStringParam(Params, TEXT("mesh_path"), MeshPath))
+	{
+		return FECACommandResult::Error(TEXT("Missing required parameter: mesh_path"));
+	}
+
+	FString OutputPathA, OutputPathB;
+	if (!GetStringParam(Params, TEXT("output_path_a"), OutputPathA))
+	{
+		return FECACommandResult::Error(TEXT("Missing required parameter: output_path_a"));
+	}
+	if (!GetStringParam(Params, TEXT("output_path_b"), OutputPathB))
+	{
+		return FECACommandResult::Error(TEXT("Missing required parameter: output_path_b"));
+	}
+
+	const TSharedPtr<FJsonObject>* PlaneOriginObj;
+	if (!GetObjectParam(Params, TEXT("plane_origin"), PlaneOriginObj))
+	{
+		return FECACommandResult::Error(TEXT("Missing required parameter: plane_origin"));
+	}
+
+	const TSharedPtr<FJsonObject>* PlaneNormalObj;
+	if (!GetObjectParam(Params, TEXT("plane_normal"), PlaneNormalObj))
+	{
+		return FECACommandResult::Error(TEXT("Missing required parameter: plane_normal"));
+	}
+
+	bool bFillHoles = true;
+	GetBoolParam(Params, TEXT("fill_holes"), bFillHoles, false);
+
+	FVector3d PlaneOrigin = MeshHelpers::ParseVector(*PlaneOriginObj);
+	FVector3d PlaneNormal = MeshHelpers::ParseVector(*PlaneNormalObj);
+	PlaneNormal = PlaneNormal.IsNearlyZero() ? FVector3d::UnitZ() : PlaneNormal.GetSafeNormal();
+
+	UStaticMesh* SourceMesh = MeshHelpers::LoadStaticMesh(MeshPath);
+	if (!SourceMesh)
+	{
+		return FECACommandResult::Error(FString::Printf(TEXT("Mesh not found: %s"), *MeshPath));
+	}
+
+	FDynamicMesh3 OriginalMesh;
+	FString Error;
+	if (!MeshHelpers::LoadStaticMeshIntoDynamicMesh(SourceMesh, OriginalMesh, &Error))
+	{
+		return FECACommandResult::Error(Error);
+	}
+
+	// Make copies for both pieces
+	FDynamicMesh3 MeshA(OriginalMesh);
+	FDynamicMesh3 MeshB(OriginalMesh);
+
+	// Cut piece A (positive side removed = keep negative side)
+	FMeshPlaneCut CutA(&MeshA, PlaneOrigin, PlaneNormal);
+	CutA.Cut();
+	if (bFillHoles)
+	{
+		for (const FMeshPlaneCut::FOpenBoundary& Boundary : CutA.OpenBoundaries)
+		{
+			for (const FEdgeLoop& Loop : Boundary.CutLoops)
+			{
+				FSimpleHoleFiller Filler(&MeshA, Loop);
+				Filler.Fill();
+			}
+		}
+	}
+	FMeshNormals::QuickComputeVertexNormals(MeshA);
+
+	// Cut piece B (flip normal to keep opposite side)
+	FMeshPlaneCut CutB(&MeshB, PlaneOrigin, -PlaneNormal);
+	CutB.Cut();
+	if (bFillHoles)
+	{
+		for (const FMeshPlaneCut::FOpenBoundary& Boundary : CutB.OpenBoundaries)
+		{
+			for (const FEdgeLoop& Loop : Boundary.CutLoops)
+			{
+				FSimpleHoleFiller Filler(&MeshB, Loop);
+				Filler.Fill();
+			}
+		}
+	}
+	FMeshNormals::QuickComputeVertexNormals(MeshB);
+
+	// Create output meshes
+	UStaticMesh* OutputMeshA = MeshHelpers::CreateStaticMeshFromDynamicMesh(MeshA, OutputPathA, true, true, &Error);
+	if (!OutputMeshA)
+	{
+		return FECACommandResult::Error(FString::Printf(TEXT("Failed to create mesh A: %s"), *Error));
+	}
+
+	UStaticMesh* OutputMeshB = MeshHelpers::CreateStaticMeshFromDynamicMesh(MeshB, OutputPathB, true, true, &Error);
+	if (!OutputMeshB)
+	{
+		return FECACommandResult::Error(FString::Printf(TEXT("Failed to create mesh B: %s"), *Error));
+	}
+
+	TSharedPtr<FJsonObject> Result = MakeResult();
+	Result->SetStringField(TEXT("mesh_a_path"), OutputMeshA->GetPathName());
+	Result->SetStringField(TEXT("mesh_b_path"), OutputMeshB->GetPathName());
+	Result->SetNumberField(TEXT("mesh_a_triangles"), MeshA.TriangleCount());
+	Result->SetNumberField(TEXT("mesh_b_triangles"), MeshB.TriangleCount());
+	
+	return FECACommandResult::Success(Result);
+}
+
+//==============================================================================
+// MESH MIRROR
+//==============================================================================
+
+FECACommandResult FECACommand_MeshMirror::Execute(const TSharedPtr<FJsonObject>& Params)
+{
+	FString MeshPath;
+	if (!GetStringParam(Params, TEXT("mesh_path"), MeshPath))
+	{
+		return FECACommandResult::Error(TEXT("Missing required parameter: mesh_path"));
+	}
+
+	FString OutputPath;
+	if (!GetStringParam(Params, TEXT("output_path"), OutputPath, false))
+	{
+		OutputPath = MeshPath;
+	}
+
+	FString MirrorAxis = TEXT("x");
+	GetStringParam(Params, TEXT("mirror_axis"), MirrorAxis, false);
+
+	bool bWeld = true;
+	GetBoolParam(Params, TEXT("weld"), bWeld, false);
+
+	UStaticMesh* SourceMesh = MeshHelpers::LoadStaticMesh(MeshPath);
+	if (!SourceMesh)
+	{
+		return FECACommandResult::Error(FString::Printf(TEXT("Mesh not found: %s"), *MeshPath));
+	}
+
+	FDynamicMesh3 DynMesh;
+	FString Error;
+	if (!MeshHelpers::LoadStaticMeshIntoDynamicMesh(SourceMesh, DynMesh, &Error))
+	{
+		return FECACommandResult::Error(Error);
+	}
+
+	// Determine mirror plane
+	FVector3d PlaneOrigin = FVector3d::Zero();
+	FVector3d PlaneNormal = FVector3d::UnitX();
+
+	const TSharedPtr<FJsonObject>* PlaneOriginObj = nullptr;
+	const TSharedPtr<FJsonObject>* PlaneNormalObj = nullptr;
+	
+	if (Params->TryGetObjectField(TEXT("plane_origin"), PlaneOriginObj) &&
+		Params->TryGetObjectField(TEXT("plane_normal"), PlaneNormalObj))
+	{
+		PlaneOrigin = MeshHelpers::ParseVector(*PlaneOriginObj);
+		PlaneNormal = MeshHelpers::ParseVector(*PlaneNormalObj);
+	}
+	else
+	{
+		MirrorAxis = MirrorAxis.ToLower();
+		if (MirrorAxis == TEXT("y"))
+		{
+			PlaneNormal = FVector3d::UnitY();
+		}
+		else if (MirrorAxis == TEXT("z"))
+		{
+			PlaneNormal = FVector3d::UnitZ();
+		}
+	}
+	PlaneNormal = PlaneNormal.IsNearlyZero() ? FVector3d::UnitX() : PlaneNormal.GetSafeNormal();
+
+	// Perform mirror
+	FMeshMirror Mirror(&DynMesh, PlaneOrigin, PlaneNormal);
+	Mirror.bWeldAlongPlane = bWeld;
+	Mirror.MirrorAndAppend();
+
+	FMeshNormals::QuickComputeVertexNormals(DynMesh);
+
+	// Save result
+	UStaticMesh* OutputMesh = nullptr;
+	if (OutputPath == MeshPath)
+	{
+		MeshHelpers::SaveDynamicMeshToStaticMesh(DynMesh, SourceMesh);
+		OutputMesh = SourceMesh;
+	}
+	else
+	{
+		OutputMesh = MeshHelpers::CreateStaticMeshFromDynamicMesh(DynMesh, OutputPath, true, true, &Error);
+		if (!OutputMesh)
+		{
+			return FECACommandResult::Error(Error);
+		}
+	}
+
+	TSharedPtr<FJsonObject> Result = MeshHelpers::BuildMeshResultJson(OutputMesh);
+	return FECACommandResult::Success(Result);
+}
+
+//==============================================================================
+// MODELING OPERATIONS
+//==============================================================================
+
+FECACommandResult FECACommand_MeshExtrude::Execute(const TSharedPtr<FJsonObject>& Params)
+{
+	return FECACommandResult::Error(TEXT("mesh_extrude: Use extrude_polygon for profile extrusion"));
+}
+
+FECACommandResult FECACommand_ExtrudePolygon::Execute(const TSharedPtr<FJsonObject>& Params)
+{
+	FString AssetPath;
+	if (!GetStringParam(Params, TEXT("asset_path"), AssetPath))
+	{
+		return FECACommandResult::Error(TEXT("Missing required parameter: asset_path"));
+	}
+
+	const TArray<TSharedPtr<FJsonValue>>* PointsArray;
+	if (!GetArrayParam(Params, TEXT("points"), PointsArray))
+	{
+		return FECACommandResult::Error(TEXT("Missing required parameter: points"));
+	}
+
+	double Height = 100;
+	if (!GetFloatParam(Params, TEXT("height"), Height))
+	{
+		return FECACommandResult::Error(TEXT("Missing required parameter: height"));
+	}
+
+	int32 Segments = 1;
+	GetIntParam(Params, TEXT("segments"), Segments, false);
+	if (Segments < 1) Segments = 1;
+
+	// Parse profile points
+	TArray<FVector2d> Points;
+	for (const TSharedPtr<FJsonValue>& PointValue : *PointsArray)
+	{
+		const TSharedPtr<FJsonObject>* PointObj;
+		if (PointValue->TryGetObject(PointObj))
+		{
+			Points.Add(FVector2d(
+				(*PointObj)->GetNumberField(TEXT("x")),
+				(*PointObj)->GetNumberField(TEXT("y"))
+			));
+		}
+	}
+
+	if (Points.Num() < 3)
+	{
+		return FECACommandResult::Error(TEXT("Profile must have at least 3 points"));
+	}
+
+	// Build extrusion path
+	FPolygon2d Profile(Points);
+	TArray<FVector3d> Path;
+	
+	for (int32 i = 0; i <= Segments; i++)
+	{
+		double T = (double)i / Segments;
+		Path.Add(FVector3d(0, 0, Height * T));
+	}
+
+	// Generate mesh
+	FDynamicMesh3 DynMesh;
+	FGeneralizedCylinderGenerator ExtrudeGen;
+	ExtrudeGen.CrossSection = Profile;
+	ExtrudeGen.Path = Path;
+	ExtrudeGen.bCapped = true;
+	ExtrudeGen.bPolygroupPerQuad = false;
+
+	ExtrudeGen.Generate();
+	DynMesh.Copy(&ExtrudeGen);
+
+	FMeshNormals::QuickComputeVertexNormals(DynMesh);
+
+	FString Error;
+	UStaticMesh* StaticMesh = MeshHelpers::CreateStaticMeshFromDynamicMesh(DynMesh, AssetPath, true, true, &Error);
+	if (!StaticMesh)
+	{
+		return FECACommandResult::Error(Error);
+	}
+
+	TSharedPtr<FJsonObject> Result = MeshHelpers::BuildMeshResultJson(StaticMesh, DynMesh.VertexCount(), DynMesh.TriangleCount());
+	return FECACommandResult::Success(Result);
+}
+
+FECACommandResult FECACommand_RevolveProfile::Execute(const TSharedPtr<FJsonObject>& Params)
+{
+	FString AssetPath;
+	if (!GetStringParam(Params, TEXT("asset_path"), AssetPath))
+	{
+		return FECACommandResult::Error(TEXT("Missing required parameter: asset_path"));
+	}
+
+	const TArray<TSharedPtr<FJsonValue>>* PointsArray;
+	if (!GetArrayParam(Params, TEXT("points"), PointsArray))
+	{
+		return FECACommandResult::Error(TEXT("Missing required parameter: points"));
+	}
+
+	double Angle = 360;
+	GetFloatParam(Params, TEXT("angle"), Angle, false);
+
+	int32 Segments = 32;
+	GetIntParam(Params, TEXT("segments"), Segments, false);
+
+	// Parse profile points (X = radius from axis, Y = height)
+	TArray<FVector2d> ProfilePoints;
+	for (const TSharedPtr<FJsonValue>& PointValue : *PointsArray)
+	{
+		const TSharedPtr<FJsonObject>* PointObj;
+		if (PointValue->TryGetObject(PointObj))
+		{
+			ProfilePoints.Add(FVector2d(
+				(*PointObj)->GetNumberField(TEXT("x")), // radius
+				(*PointObj)->GetNumberField(TEXT("y"))  // height
+			));
+		}
+	}
+
+	if (ProfilePoints.Num() < 2)
+	{
+		return FECACommandResult::Error(TEXT("Profile must have at least 2 points"));
+	}
+
+	// Generate revolution mesh
+	FRevolvePlanarPathGenerator RevolveGen;
+	RevolveGen.PathVertices = ProfilePoints;
+	RevolveGen.RevolveDegrees = Angle;
+	RevolveGen.Steps = Segments;
+	RevolveGen.bCapped = FMath::Abs(Angle) < 360.0 - 0.01;
+	
+	FDynamicMesh3 DynMesh = RevolveGen.GenerateMesh();
+
+
+	FMeshNormals::QuickComputeVertexNormals(DynMesh);
+
+	FString Error;
+	UStaticMesh* StaticMesh = MeshHelpers::CreateStaticMeshFromDynamicMesh(DynMesh, AssetPath, true, true, &Error);
+	if (!StaticMesh)
+	{
+		return FECACommandResult::Error(Error);
+	}
+
+	TSharedPtr<FJsonObject> Result = MeshHelpers::BuildMeshResultJson(StaticMesh, DynMesh.VertexCount(), DynMesh.TriangleCount());
+	return FECACommandResult::Success(Result);
+}
+
+FECACommandResult FECACommand_SweepProfile::Execute(const TSharedPtr<FJsonObject>& Params)
+{
+	FString AssetPath;
+	if (!GetStringParam(Params, TEXT("asset_path"), AssetPath))
+	{
+		return FECACommandResult::Error(TEXT("Missing required parameter: asset_path"));
+	}
+
+	const TArray<TSharedPtr<FJsonValue>>* ProfilePointsArray;
+	if (!GetArrayParam(Params, TEXT("profile_points"), ProfilePointsArray))
+	{
+		return FECACommandResult::Error(TEXT("Missing required parameter: profile_points"));
+	}
+
+	const TArray<TSharedPtr<FJsonValue>>* PathPointsArray;
+	if (!GetArrayParam(Params, TEXT("path_points"), PathPointsArray))
+	{
+		return FECACommandResult::Error(TEXT("Missing required parameter: path_points"));
+	}
+
+	bool bClosedPath = false;
+	GetBoolParam(Params, TEXT("closed_path"), bClosedPath, false);
+
+	// Parse profile
+	TArray<FVector2d> ProfilePoints;
+	for (const TSharedPtr<FJsonValue>& PointValue : *ProfilePointsArray)
+	{
+		const TSharedPtr<FJsonObject>* PointObj;
+		if (PointValue->TryGetObject(PointObj))
+		{
+			ProfilePoints.Add(FVector2d(
+				(*PointObj)->GetNumberField(TEXT("x")),
+				(*PointObj)->GetNumberField(TEXT("y"))
+			));
+		}
+	}
+
+	// Parse path
+	TArray<FVector3d> PathPoints;
+	for (const TSharedPtr<FJsonValue>& PointValue : *PathPointsArray)
+	{
+		const TSharedPtr<FJsonObject>* PointObj;
+		if (PointValue->TryGetObject(PointObj))
+		{
+			PathPoints.Add(FVector3d(
+				(*PointObj)->GetNumberField(TEXT("x")),
+				(*PointObj)->GetNumberField(TEXT("y")),
+				(*PointObj)->GetNumberField(TEXT("z"))
+			));
+		}
+	}
+
+	if (ProfilePoints.Num() < 3)
+	{
+		return FECACommandResult::Error(TEXT("Profile must have at least 3 points"));
+	}
+	if (PathPoints.Num() < 2)
+	{
+		return FECACommandResult::Error(TEXT("Path must have at least 2 points"));
+	}
+
+	// Generate sweep
+	FPolygon2d Profile(ProfilePoints);
+	
+	FDynamicMesh3 DynMesh;
+	FGeneralizedCylinderGenerator SweepGen;
+	SweepGen.CrossSection = Profile;
+	SweepGen.Path = PathPoints;
+	SweepGen.bLoop = bClosedPath;
+	SweepGen.bCapped = !bClosedPath;
+	SweepGen.bPolygroupPerQuad = false;
+
+	SweepGen.Generate();
+	DynMesh.Copy(&SweepGen);
+
+	FMeshNormals::QuickComputeVertexNormals(DynMesh);
+
+	FString Error;
+	UStaticMesh* StaticMesh = MeshHelpers::CreateStaticMeshFromDynamicMesh(DynMesh, AssetPath, true, true, &Error);
+	if (!StaticMesh)
+	{
+		return FECACommandResult::Error(Error);
+	}
+
+	TSharedPtr<FJsonObject> Result = MeshHelpers::BuildMeshResultJson(StaticMesh, DynMesh.VertexCount(), DynMesh.TriangleCount());
+	return FECACommandResult::Success(Result);
+}
+
+FECACommandResult FECACommand_MeshInset::Execute(const TSharedPtr<FJsonObject>& Params)
+{
+	FString MeshPath;
+	if (!GetStringParam(Params, TEXT("mesh_path"), MeshPath))
+	{
+		return FECACommandResult::Error(TEXT("Missing required parameter: mesh_path"));
+	}
+
+	FString OutputPath;
+	if (!GetStringParam(Params, TEXT("output_path"), OutputPath, false))
+	{
+		OutputPath = MeshPath;
+	}
+
+	double InsetDistance = 5.0;
+	GetFloatParam(Params, TEXT("distance"), InsetDistance, false);
+
+	UStaticMesh* SourceMesh = MeshHelpers::LoadStaticMesh(MeshPath);
+	if (!SourceMesh)
+	{
+		return FECACommandResult::Error(FString::Printf(TEXT("Mesh not found: %s"), *MeshPath));
+	}
+
+	FDynamicMesh3 DynMesh;
+	FString Error;
+	if (!MeshHelpers::LoadStaticMeshIntoDynamicMesh(SourceMesh, DynMesh, &Error))
+	{
+		return FECACommandResult::Error(Error);
+	}
+
+	// Simple inset: scale each face towards its centroid
+	// Collect all triangles first
+	TArray<int32> TriangleIDs;
+	for (int32 TriID : DynMesh.TriangleIndicesItr())
+	{
+		TriangleIDs.Add(TriID);
+	}
+
+	// Compute face centroids and move vertices towards them
+	for (int32 TriID : TriangleIDs)
+	{
+		if (!DynMesh.IsTriangle(TriID)) continue;
+		
+		FIndex3i Tri = DynMesh.GetTriangle(TriID);
+		FVector3d V0 = DynMesh.GetVertex(Tri.A);
+		FVector3d V1 = DynMesh.GetVertex(Tri.B);
+		FVector3d V2 = DynMesh.GetVertex(Tri.C);
+		FVector3d Centroid = (V0 + V1 + V2) / 3.0;
+		
+		// Compute inset factor based on distance
+		double EdgeLen = ((V1 - V0).Length() + (V2 - V1).Length() + (V0 - V2).Length()) / 3.0;
+		double InsetFactor = FMath::Clamp(InsetDistance / EdgeLen, 0.0, 0.5);
+		
+		// Move each vertex towards centroid
+		DynMesh.SetVertex(Tri.A, FMath::Lerp(V0, Centroid, InsetFactor));
+		DynMesh.SetVertex(Tri.B, FMath::Lerp(V1, Centroid, InsetFactor));
+		DynMesh.SetVertex(Tri.C, FMath::Lerp(V2, Centroid, InsetFactor));
+	}
+
+	FMeshNormals::QuickComputeVertexNormals(DynMesh);
+
+	UStaticMesh* OutputMesh = nullptr;
+	if (OutputPath == MeshPath)
+	{
+		MeshHelpers::SaveDynamicMeshToStaticMesh(DynMesh, SourceMesh);
+		OutputMesh = SourceMesh;
+	}
+	else
+	{
+		OutputMesh = MeshHelpers::CreateStaticMeshFromDynamicMesh(DynMesh, OutputPath, true, true, &Error);
+		if (!OutputMesh)
+		{
+			return FECACommandResult::Error(Error);
+		}
+	}
+
+	TSharedPtr<FJsonObject> Result = MeshHelpers::BuildMeshResultJson(OutputMesh);
+	Result->SetNumberField(TEXT("inset_distance"), InsetDistance);
+	return FECACommandResult::Success(Result);
+}
+
+
+
+FECACommandResult FECACommand_MeshOffset::Execute(const TSharedPtr<FJsonObject>& Params)
+{
+	FString MeshPath;
+	if (!GetStringParam(Params, TEXT("mesh_path"), MeshPath))
+	{
+		return FECACommandResult::Error(TEXT("Missing required parameter: mesh_path"));
+	}
+
+	double Distance = 0;
+	if (!GetFloatParam(Params, TEXT("distance"), Distance))
+	{
+		return FECACommandResult::Error(TEXT("Missing required parameter: distance"));
+	}
+
+	FString OutputPath;
+	if (!GetStringParam(Params, TEXT("output_path"), OutputPath, false))
+	{
+		OutputPath = MeshPath;
+	}
+
+	UStaticMesh* SourceMesh = MeshHelpers::LoadStaticMesh(MeshPath);
+	if (!SourceMesh)
+	{
+		return FECACommandResult::Error(FString::Printf(TEXT("Mesh not found: %s"), *MeshPath));
+	}
+
+	FDynamicMesh3 DynMesh;
+	FString Error;
+	if (!MeshHelpers::LoadStaticMeshIntoDynamicMesh(SourceMesh, DynMesh, &Error))
+	{
+		return FECACommandResult::Error(Error);
+	}
+
+	// Compute normals first
+	FMeshNormals Normals(&DynMesh);
+	Normals.ComputeVertexNormals();
+	
+	// Simple offset along normals
+	for (int32 VertID : DynMesh.VertexIndicesItr())
+	{
+		FVector3d Normal = Normals[VertID];
+		FVector3d Pos = DynMesh.GetVertex(VertID);
+		DynMesh.SetVertex(VertID, Pos + Normal * Distance);
+	}
+
+	FMeshNormals::QuickComputeVertexNormals(DynMesh);
+
+	UStaticMesh* OutputMesh = nullptr;
+	if (OutputPath == MeshPath)
+	{
+		MeshHelpers::SaveDynamicMeshToStaticMesh(DynMesh, SourceMesh);
+		OutputMesh = SourceMesh;
+	}
+	else
+	{
+		OutputMesh = MeshHelpers::CreateStaticMeshFromDynamicMesh(DynMesh, OutputPath, true, true, &Error);
+		if (!OutputMesh)
+		{
+			return FECACommandResult::Error(Error);
+		}
+	}
+
+	TSharedPtr<FJsonObject> Result = MeshHelpers::BuildMeshResultJson(OutputMesh);
+	Result->SetNumberField(TEXT("distance"), Distance);
+	return FECACommandResult::Success(Result);
+}
+
+//==============================================================================
+// MESH MODIFICATION
+//==============================================================================
+
+
+
+FECACommandResult FECACommand_MeshSimplify::Execute(const TSharedPtr<FJsonObject>& Params)
+{
+	FString MeshPath;
+	if (!GetStringParam(Params, TEXT("mesh_path"), MeshPath))
+	{
+		return FECACommandResult::Error(TEXT("Missing required parameter: mesh_path"));
+	}
+
+	FString Method;
+	if (!GetStringParam(Params, TEXT("method"), Method))
+	{
+		return FECACommandResult::Error(TEXT("Missing required parameter: method"));
+	}
+
+	double Target = 0;
+	if (!GetFloatParam(Params, TEXT("target"), Target))
+	{
+		return FECACommandResult::Error(TEXT("Missing required parameter: target"));
+	}
+
+	FString OutputPath;
+	if (!GetStringParam(Params, TEXT("output_path"), OutputPath, false))
+	{
+		OutputPath = MeshPath;
+	}
+
+	UStaticMesh* SourceMesh = MeshHelpers::LoadStaticMesh(MeshPath);
+	if (!SourceMesh)
+	{
+		return FECACommandResult::Error(FString::Printf(TEXT("Mesh not found: %s"), *MeshPath));
+	}
+
+	FDynamicMesh3 DynMesh;
+	FString Error;
+	if (!MeshHelpers::LoadStaticMeshIntoDynamicMesh(SourceMesh, DynMesh, &Error))
+	{
+		return FECACommandResult::Error(Error);
+	}
+
+	int32 OriginalTriCount = DynMesh.TriangleCount();
+
+	// Simplify using QEM
+	FQEMSimplification Simplifier(&DynMesh);
+	
+	Method = Method.ToLower();
+	if (Method == TEXT("triangle_count"))
+	{
+		Simplifier.SimplifyToTriangleCount(FMath::Max(4, (int32)Target));
+	}
+	else if (Method == TEXT("vertex_count"))
+	{
+		Simplifier.SimplifyToVertexCount(FMath::Max(4, (int32)Target));
+	}
+	else if (Method == TEXT("percentage"))
+	{
+		int32 TargetTriCount = FMath::Max(4, (int32)(OriginalTriCount * Target / 100.0));
+		Simplifier.SimplifyToTriangleCount(TargetTriCount);
+	}
+	else if (Method == TEXT("tolerance"))
+	{
+		Simplifier.SimplifyToMaxError(Target);
+	}
+	else
+	{
+		return FECACommandResult::Error(FString::Printf(TEXT("Unknown simplification method: %s"), *Method));
+	}
+
+	FMeshNormals::QuickComputeVertexNormals(DynMesh);
+
+	UStaticMesh* OutputMesh = nullptr;
+	if (OutputPath == MeshPath)
+	{
+		MeshHelpers::SaveDynamicMeshToStaticMesh(DynMesh, SourceMesh);
+		OutputMesh = SourceMesh;
+	}
+	else
+	{
+		OutputMesh = MeshHelpers::CreateStaticMeshFromDynamicMesh(DynMesh, OutputPath, true, true, &Error);
+		if (!OutputMesh)
+		{
+			return FECACommandResult::Error(Error);
+		}
+	}
+
+	TSharedPtr<FJsonObject> Result = MeshHelpers::BuildMeshResultJson(OutputMesh);
+	Result->SetNumberField(TEXT("original_triangles"), OriginalTriCount);
+	Result->SetNumberField(TEXT("reduction_ratio"), 1.0 - (double)DynMesh.TriangleCount() / OriginalTriCount);
+	return FECACommandResult::Success(Result);
+}
+
+FECACommandResult FECACommand_MeshSubdivide::Execute(const TSharedPtr<FJsonObject>& Params)
+{
+	FString MeshPath;
+	if (!GetStringParam(Params, TEXT("mesh_path"), MeshPath))
+	{
+		return FECACommandResult::Error(TEXT("Missing required parameter: mesh_path"));
+	}
+
+	FString OutputPath;
+	if (!GetStringParam(Params, TEXT("output_path"), OutputPath, false))
+	{
+		OutputPath = MeshPath;
+	}
+
+	int32 Iterations = 1;
+	GetIntParam(Params, TEXT("iterations"), Iterations, false);
+	Iterations = FMath::Clamp(Iterations, 1, 5); // Limit to prevent excessive geometry
+
+	FString Method = TEXT("poke");
+	GetStringParam(Params, TEXT("method"), Method, false);
+	Method = Method.ToLower();
+
+	UStaticMesh* SourceMesh = MeshHelpers::LoadStaticMesh(MeshPath);
+	if (!SourceMesh)
+	{
+		return FECACommandResult::Error(FString::Printf(TEXT("Mesh not found: %s"), *MeshPath));
+	}
+
+	FDynamicMesh3 DynMesh;
+	FString Error;
+	if (!MeshHelpers::LoadStaticMeshIntoDynamicMesh(SourceMesh, DynMesh, &Error))
+	{
+		return FECACommandResult::Error(Error);
+	}
+
+	int32 OriginalTriCount = DynMesh.TriangleCount();
+
+	// Use proper subdivision algorithms via FSubdividePoly (OpenSubdiv-based)
+	if (Method == TEXT("loop") || Method == TEXT("catmull_clark") || Method == TEXT("catmullclark") || Method == TEXT("bilinear"))
+	{
+		// Build group topology from the mesh (required for subdivision)
+		FGroupTopology Topology(&DynMesh, true);
+		
+		// Create subdivider
+		FSubdividePoly Subdivider(Topology, DynMesh, Iterations);
+		
+		// Set subdivision scheme
+		if (Method == TEXT("loop"))
+		{
+			Subdivider.SubdivisionScheme = ESubdivisionScheme::Loop;
+		}
+		else if (Method == TEXT("catmull_clark") || Method == TEXT("catmullclark"))
+		{
+			Subdivider.SubdivisionScheme = ESubdivisionScheme::CatmullClark;
+		}
+		else // bilinear
+		{
+			Subdivider.SubdivisionScheme = ESubdivisionScheme::Bilinear;
+		}
+		
+		Subdivider.NormalComputationMethod = ESubdivisionOutputNormals::Generated;
+		Subdivider.UVComputationMethod = ESubdivisionOutputUVs::Interpolated;
+		Subdivider.BoundaryScheme = ESubdivisionBoundaryScheme::SharpCorners;
+		
+		// Validate topology
+		FSubdividePoly::ETopologyCheckResult CheckResult = Subdivider.ValidateTopology();
+		if (CheckResult != FSubdividePoly::ETopologyCheckResult::Ok)
+		{
+			// Fall back to simple uniform subdivision for problematic topology
+			FString Warning = FString::Printf(TEXT("Topology validation failed (%d), falling back to uniform subdivision"), (int32)CheckResult);
+			UE_LOG(LogTemp, Warning, TEXT("%s"), *Warning);
+			
+			// Use uniform tessellation as fallback
+			FUniformTessellate Tessellator(&DynMesh);
+			Tessellator.TessellationNum = Iterations + 1;
+			if (!Tessellator.Compute())
+			{
+				return FECACommandResult::Error(TEXT("Uniform tessellation failed"));
+			}
+		}
+		else
+		{
+			// Compute subdivision
+			if (!Subdivider.ComputeTopologySubdivision())
+			{
+				return FECACommandResult::Error(TEXT("Failed to compute topology subdivision"));
+			}
+			
+			FDynamicMesh3 SubdividedMesh;
+			if (!Subdivider.ComputeSubdividedMesh(SubdividedMesh))
+			{
+				return FECACommandResult::Error(TEXT("Failed to compute subdivided mesh"));
+			}
+			
+			DynMesh = MoveTemp(SubdividedMesh);
+		}
+	}
+	else if (Method == TEXT("uniform") || Method == TEXT("tessellate"))
+	{
+		// Uniform tessellation - splits triangles uniformly
+		FUniformTessellate Tessellator(&DynMesh);
+		Tessellator.TessellationNum = Iterations + 1;  // TessellationNum of 2 = 1 subdivision level
+		if (!Tessellator.Compute())
+		{
+			return FECACommandResult::Error(TEXT("Uniform tessellation failed"));
+		}
+	}
+	else if (Method == TEXT("pn") || Method == TEXT("pn_triangles"))
+	{
+		// PN Triangles - curved surface approximation
+		FPNTriangles PNTess(&DynMesh);
+		PNTess.TessellationLevel = Iterations;
+		if (!PNTess.Compute())
+		{
+			return FECACommandResult::Error(TEXT("PN Triangles tessellation failed"));
+		}
+	}
+	else
+	{
+		return FECACommandResult::Error(FString::Printf(TEXT("Unknown subdivision method: %s. Use 'loop', 'catmull_clark', 'bilinear', 'uniform', or 'pn_triangles'."), *Method));
+	}
+
+	FMeshNormals::QuickComputeVertexNormals(DynMesh);
+
+	UStaticMesh* OutputMesh = nullptr;
+	if (OutputPath == MeshPath)
+	{
+		MeshHelpers::SaveDynamicMeshToStaticMesh(DynMesh, SourceMesh);
+		OutputMesh = SourceMesh;
+	}
+	else
+	{
+		OutputMesh = MeshHelpers::CreateStaticMeshFromDynamicMesh(DynMesh, OutputPath, true, true, &Error);
+		if (!OutputMesh)
+		{
+			return FECACommandResult::Error(Error);
+		}
+	}
+
+	TSharedPtr<FJsonObject> Result = MeshHelpers::BuildMeshResultJson(OutputMesh);
+	Result->SetNumberField(TEXT("original_triangles"), OriginalTriCount);
+	Result->SetNumberField(TEXT("new_triangles"), DynMesh.TriangleCount());
+	Result->SetNumberField(TEXT("iterations"), Iterations);
+	Result->SetStringField(TEXT("method"), Method);
+	return FECACommandResult::Success(Result);
+}
+
+FECACommandResult FECACommand_MeshRemesh::Execute(const TSharedPtr<FJsonObject>& Params)
+{
+	FString MeshPath;
+	if (!GetStringParam(Params, TEXT("mesh_path"), MeshPath))
+	{
+		return FECACommandResult::Error(TEXT("Missing required parameter: mesh_path"));
+	}
+
+	double TargetEdgeLength = 10;
+	if (!GetFloatParam(Params, TEXT("target_edge_length"), TargetEdgeLength))
+	{
+		return FECACommandResult::Error(TEXT("Missing required parameter: target_edge_length"));
+	}
+
+	FString OutputPath;
+	if (!GetStringParam(Params, TEXT("output_path"), OutputPath, false))
+	{
+		OutputPath = MeshPath;
+	}
+
+	int32 Iterations = 20;
+	GetIntParam(Params, TEXT("iterations"), Iterations, false);
+
+	double SmoothingRate = 0.25;
+	GetFloatParam(Params, TEXT("smoothing_rate"), SmoothingRate, false);
+
+	UStaticMesh* SourceMesh = MeshHelpers::LoadStaticMesh(MeshPath);
+	if (!SourceMesh)
+	{
+		return FECACommandResult::Error(FString::Printf(TEXT("Mesh not found: %s"), *MeshPath));
+	}
+
+	FDynamicMesh3 DynMesh;
+	FString Error;
+	if (!MeshHelpers::LoadStaticMeshIntoDynamicMesh(SourceMesh, DynMesh, &Error))
+	{
+		return FECACommandResult::Error(Error);
+	}
+
+	// Validate mesh before remeshing
+	if (DynMesh.TriangleCount() < 1)
+	{
+		return FECACommandResult::Error(TEXT("Mesh has no triangles"));
+	}
+
+	// Clamp parameters to safe ranges
+	TargetEdgeLength = FMath::Max(0.1, TargetEdgeLength);
+	Iterations = FMath::Clamp(Iterations, 1, 100);
+	SmoothingRate = FMath::Clamp(SmoothingRate, 0.0, 1.0);
+
+	// Split any bowties in attribute layers before remeshing
+	if (DynMesh.HasAttributes())
+	{
+		DynMesh.Attributes()->SplitAllBowties();
+	}
+
+	// Remesh using FRemesher
+	FRemesher Remesher(&DynMesh);
+	
+	// Set up constraints for attribute seams if the mesh has attributes
+	if (DynMesh.HasAttributes())
+	{
+		FMeshConstraints Constraints;
+		FMeshConstraintsUtil::ConstrainAllBoundariesAndSeams(
+			Constraints, 
+			DynMesh,
+			EEdgeRefineFlags::NoFlip,           // MeshBoundaryConstraints
+			EEdgeRefineFlags::NoConstraint,     // GroupBorderConstraints  
+			EEdgeRefineFlags::NoConstraint,     // MaterialBorderConstraints
+			true,   // bAllowSeamSplits
+			true,   // bAllowSeamSmoothing
+			true    // bAllowSeamCollapse
+		);
+		Remesher.SetExternalConstraints(MoveTemp(Constraints));
+	}
+	
+	// Set target edge length range (typically MinEdgeLength = 0.5*Target, MaxEdgeLength = 1.5*Target)
+	Remesher.SetTargetEdgeLength(TargetEdgeLength);
+	
+	// Configure smoothing
+	Remesher.SmoothSpeedT = SmoothingRate;
+	Remesher.bEnableSmoothing = SmoothingRate > 0;
+	
+	// Enable standard operations
+	Remesher.bEnableFlips = true;
+	Remesher.bEnableCollapses = true;
+	Remesher.bEnableSplits = true;
+	Remesher.bPreventNormalFlips = true;
+	
+	// Perform remeshing iterations with error handling
+	for (int32 i = 0; i < Iterations; i++)
+	{
+		if (DynMesh.TriangleCount() < 4)
+		{
+			// Mesh has become too small, stop
+			break;
+		}
+		Remesher.BasicRemeshPass();
+	}
+
+	FMeshNormals::QuickComputeVertexNormals(DynMesh);
+
+	UStaticMesh* OutputMesh = nullptr;
+	if (OutputPath == MeshPath)
+	{
+		MeshHelpers::SaveDynamicMeshToStaticMesh(DynMesh, SourceMesh);
+		OutputMesh = SourceMesh;
+	}
+	else
+	{
+		OutputMesh = MeshHelpers::CreateStaticMeshFromDynamicMesh(DynMesh, OutputPath, true, true, &Error);
+		if (!OutputMesh)
+		{
+			return FECACommandResult::Error(Error);
+		}
+	}
+
+	TSharedPtr<FJsonObject> Result = MeshHelpers::BuildMeshResultJson(OutputMesh);
+	return FECACommandResult::Success(Result);
+}
+
+FECACommandResult FECACommand_MeshSmooth::Execute(const TSharedPtr<FJsonObject>& Params)
+{
+	FString MeshPath;
+	if (!GetStringParam(Params, TEXT("mesh_path"), MeshPath))
+	{
+		return FECACommandResult::Error(TEXT("Missing required parameter: mesh_path"));
+	}
+
+	FString OutputPath;
+	if (!GetStringParam(Params, TEXT("output_path"), OutputPath, false))
+	{
+		OutputPath = MeshPath;
+	}
+
+	double Strength = 0.5;
+	GetFloatParam(Params, TEXT("strength"), Strength, false);
+
+	int32 Iterations = 1;
+	GetIntParam(Params, TEXT("iterations"), Iterations, false);
+
+	UStaticMesh* SourceMesh = MeshHelpers::LoadStaticMesh(MeshPath);
+	if (!SourceMesh)
+	{
+		return FECACommandResult::Error(FString::Printf(TEXT("Mesh not found: %s"), *MeshPath));
+	}
+
+	FDynamicMesh3 DynMesh;
+	FString Error;
+	if (!MeshHelpers::LoadStaticMeshIntoDynamicMesh(SourceMesh, DynMesh, &Error))
+	{
+		return FECACommandResult::Error(Error);
+	}
+
+	// Build a map of co-located vertices (vertices at the same position that should move together)
+	// This handles seam vertices, hard edge splits, etc.
+	const double WeldTolerance = 0.0001;
+	TMap<int32, TArray<int32>> ColocatedGroups;  // Representative vertex -> all co-located vertices
+	TMap<int32, int32> VertToRepresentative;     // Each vertex -> its representative
+	
+	{
+		// Use a spatial hash to find co-located vertices efficiently
+		TMap<FIntVector, TArray<int32>> SpatialHash;
+		const double CellSize = WeldTolerance * 2.0;
+		
+		for (int32 VertID : DynMesh.VertexIndicesItr())
+		{
+			FVector3d Pos = DynMesh.GetVertex(VertID);
+			FIntVector Cell(
+				FMath::FloorToInt(Pos.X / CellSize),
+				FMath::FloorToInt(Pos.Y / CellSize),
+				FMath::FloorToInt(Pos.Z / CellSize)
+			);
+			SpatialHash.FindOrAdd(Cell).Add(VertID);
+		}
+		
+		// Find co-located vertices within each cell and neighboring cells
+		TSet<int32> Processed;
+		for (int32 VertID : DynMesh.VertexIndicesItr())
+		{
+			if (Processed.Contains(VertID)) continue;
+			
+			FVector3d Pos = DynMesh.GetVertex(VertID);
+			FIntVector Cell(
+				FMath::FloorToInt(Pos.X / CellSize),
+				FMath::FloorToInt(Pos.Y / CellSize),
+				FMath::FloorToInt(Pos.Z / CellSize)
+			);
+			
+			TArray<int32> Group;
+			Group.Add(VertID);
+			
+			// Check neighboring cells too
+			for (int32 dx = -1; dx <= 1; dx++)
+			{
+				for (int32 dy = -1; dy <= 1; dy++)
+				{
+					for (int32 dz = -1; dz <= 1; dz++)
+					{
+						FIntVector NeighborCell(Cell.X + dx, Cell.Y + dy, Cell.Z + dz);
+						if (TArray<int32>* CellVerts = SpatialHash.Find(NeighborCell))
+						{
+							for (int32 OtherID : *CellVerts)
+							{
+								if (OtherID != VertID && !Processed.Contains(OtherID))
+								{
+									FVector3d OtherPos = DynMesh.GetVertex(OtherID);
+									if ((Pos - OtherPos).SquaredLength() < WeldTolerance * WeldTolerance)
+									{
+										Group.Add(OtherID);
+										Processed.Add(OtherID);
+									}
+								}
+							}
+						}
+					}
+				}
+			}
+			
+			Processed.Add(VertID);
+			
+			// VertID is the representative for this group
+			for (int32 GroupVert : Group)
+			{
+				VertToRepresentative.Add(GroupVert, VertID);
+			}
+			if (Group.Num() > 1)
+			{
+				ColocatedGroups.Add(VertID, MoveTemp(Group));
+			}
+		}
+	}
+	
+	// HC (Humphrey's Classes) smoothing - prevents shrinkage while smoothing
+	// Reference: "Improved Laplacian Smoothing of Noisy Surface Meshes" by Vollmer et al.
+	// 
+	// HC smoothing works in two steps per iteration:
+	// 1. Apply Laplacian smoothing to get positions Q
+	// 2. Compute difference vectors B = Q - (alpha * OriginalPos + (1-alpha) * PreviousPos)
+	// 3. Apply Laplacian smoothing to B to get correction
+	// 4. Final position = Q - (beta * B + (1-beta) * avg(neighbor B))
+	//
+	// Alpha controls how much the original position influences the correction (0.0 = use previous, 1.0 = use original)
+	// Beta controls how much local vs neighbor correction is used (0.0 = full neighbor averaging, 1.0 = only local)
+	// 
+	// Typical values: alpha = 0.0, beta = 0.5
+	
+	const double Alpha = 0.0;  // Use previous position for reference
+	const double Beta = 0.5;   // Balance between local and neighbor correction
+	
+	// Store original positions for reference
+	TMap<int32, FVector3d> OriginalPositions;
+	for (int32 VertID : DynMesh.VertexIndicesItr())
+	{
+		int32 RepID = VertToRepresentative.FindRef(VertID);
+		if (!OriginalPositions.Contains(RepID))
+		{
+			OriginalPositions.Add(RepID, DynMesh.GetVertex(RepID));
+		}
+	}
+	
+	// Helper lambda to compute Laplacian (average of neighbors) for a representative vertex
+	auto ComputeLaplacian = [&](int32 RepID) -> FVector3d
+	{
+		FVector3d Centroid = FVector3d::Zero();
+		int32 Count = 0;
+		TSet<int32> NeighborReps;
+		
+		TArray<int32>* Group = ColocatedGroups.Find(RepID);
+		if (Group)
+		{
+			for (int32 GroupVert : *Group)
+			{
+				for (int32 NeighborID : DynMesh.VtxVerticesItr(GroupVert))
+				{
+					int32 NeighborRep = VertToRepresentative.FindRef(NeighborID);
+					if (NeighborRep != RepID && !NeighborReps.Contains(NeighborRep))
+					{
+						NeighborReps.Add(NeighborRep);
+						Centroid += DynMesh.GetVertex(NeighborRep);
+						Count++;
+					}
+				}
+			}
+		}
+		else
+		{
+			for (int32 NeighborID : DynMesh.VtxVerticesItr(RepID))
+			{
+				int32 NeighborRep = VertToRepresentative.FindRef(NeighborID);
+				if (!NeighborReps.Contains(NeighborRep))
+				{
+					NeighborReps.Add(NeighborRep);
+					Centroid += DynMesh.GetVertex(NeighborRep);
+					Count++;
+				}
+			}
+		}
+		
+		return (Count > 0) ? (Centroid / Count) : DynMesh.GetVertex(RepID);
+	};
+	
+	// Helper to get neighbor representatives
+	auto GetNeighborReps = [&](int32 RepID) -> TArray<int32>
+	{
+		TArray<int32> Result;
+		TSet<int32> NeighborReps;
+		
+		TArray<int32>* Group = ColocatedGroups.Find(RepID);
+		if (Group)
+		{
+			for (int32 GroupVert : *Group)
+			{
+				for (int32 NeighborID : DynMesh.VtxVerticesItr(GroupVert))
+				{
+					int32 NeighborRep = VertToRepresentative.FindRef(NeighborID);
+					if (NeighborRep != RepID && !NeighborReps.Contains(NeighborRep))
+					{
+						NeighborReps.Add(NeighborRep);
+						Result.Add(NeighborRep);
+					}
+				}
+			}
+		}
+		else
+		{
+			for (int32 NeighborID : DynMesh.VtxVerticesItr(RepID))
+			{
+				int32 NeighborRep = VertToRepresentative.FindRef(NeighborID);
+				if (!NeighborReps.Contains(NeighborRep))
+				{
+					NeighborReps.Add(NeighborRep);
+					Result.Add(NeighborRep);
+				}
+			}
+		}
+		return Result;
+	};
+	
+	for (int32 iter = 0; iter < Iterations; iter++)
+	{
+		// Step 1: Compute Laplacian smoothed positions Q
+		TMap<int32, FVector3d> Q;  // RepID -> smoothed position
+		TMap<int32, FVector3d> PrevPos;  // RepID -> position before this iteration
+		
+		TSet<int32> AllReps;
+		for (int32 VertID : DynMesh.VertexIndicesItr())
+		{
+			int32 RepID = VertToRepresentative.FindRef(VertID);
+			AllReps.Add(RepID);
+		}
+		
+		for (int32 RepID : AllReps)
+		{
+			FVector3d CurrentPos = DynMesh.GetVertex(RepID);
+			PrevPos.Add(RepID, CurrentPos);
+			
+			FVector3d Laplacian = ComputeLaplacian(RepID);
+			FVector3d SmoothedPos = FMath::Lerp(CurrentPos, Laplacian, Strength);
+			Q.Add(RepID, SmoothedPos);
+		}
+		
+		// Step 2: Compute difference vectors B
+		// B = Q - (alpha * Original + (1-alpha) * Previous)
+		TMap<int32, FVector3d> B;
+		for (int32 RepID : AllReps)
+		{
+			FVector3d Original = OriginalPositions.FindRef(RepID);
+			FVector3d Previous = PrevPos.FindRef(RepID);
+			FVector3d Reference = Alpha * Original + (1.0 - Alpha) * Previous;
+			B.Add(RepID, Q.FindRef(RepID) - Reference);
+		}
+		
+		// Step 3: Compute corrected positions
+		// FinalPos = Q - (beta * B[i] + (1-beta) * avg(B[neighbors]))
+		TArray<FVector3d> NewPositions;
+		NewPositions.SetNum(DynMesh.MaxVertexID());
+		
+		for (int32 RepID : AllReps)
+		{
+			// Compute average B of neighbors
+			FVector3d AvgNeighborB = FVector3d::Zero();
+			TArray<int32> Neighbors = GetNeighborReps(RepID);
+			if (Neighbors.Num() > 0)
+			{
+				for (int32 NeighborRep : Neighbors)
+				{
+					AvgNeighborB += B.FindRef(NeighborRep);
+				}
+				AvgNeighborB /= Neighbors.Num();
+			}
+			
+			FVector3d LocalB = B.FindRef(RepID);
+			FVector3d Correction = Beta * LocalB + (1.0 - Beta) * AvgNeighborB;
+			FVector3d FinalPos = Q.FindRef(RepID) - Correction;
+			
+			// Apply to all co-located vertices
+			TArray<int32>* Group = ColocatedGroups.Find(RepID);
+			if (Group)
+			{
+				for (int32 GroupVert : *Group)
+				{
+					NewPositions[GroupVert] = FinalPos;
+				}
+			}
+			else
+			{
+				NewPositions[RepID] = FinalPos;
+			}
+		}
+		
+		// Apply new positions
+		for (int32 VertID : DynMesh.VertexIndicesItr())
+		{
+			DynMesh.SetVertex(VertID, NewPositions[VertID]);
+		}
+	}
+
+	FMeshNormals::QuickComputeVertexNormals(DynMesh);
+
+	UStaticMesh* OutputMesh = nullptr;
+	if (OutputPath == MeshPath)
+	{
+		MeshHelpers::SaveDynamicMeshToStaticMesh(DynMesh, SourceMesh);
+		OutputMesh = SourceMesh;
+	}
+	else
+	{
+		OutputMesh = MeshHelpers::CreateStaticMeshFromDynamicMesh(DynMesh, OutputPath, true, true, &Error);
+		if (!OutputMesh)
+		{
+			return FECACommandResult::Error(Error);
+		}
+	}
+
+	TSharedPtr<FJsonObject> Result = MeshHelpers::BuildMeshResultJson(OutputMesh);
+	return FECACommandResult::Success(Result);
+}
+
+FECACommandResult FECACommand_MeshDisplace::Execute(const TSharedPtr<FJsonObject>& Params)
+{
+	FString MeshPath;
+	if (!GetStringParam(Params, TEXT("mesh_path"), MeshPath))
+	{
+		return FECACommandResult::Error(TEXT("Missing required parameter: mesh_path"));
+	}
+
+	FString OutputPath;
+	if (!GetStringParam(Params, TEXT("output_path"), OutputPath, false))
+	{
+		OutputPath = MeshPath;
+	}
+
+	double Amount = 10.0;
+	GetFloatParam(Params, TEXT("amount"), Amount, false);
+
+	// Displacement mode: "normal", "random", "noise"
+	FString Mode = TEXT("normal");
+	GetStringParam(Params, TEXT("mode"), Mode, false);
+	Mode = Mode.ToLower();
+
+	UStaticMesh* SourceMesh = MeshHelpers::LoadStaticMesh(MeshPath);
+	if (!SourceMesh)
+	{
+		return FECACommandResult::Error(FString::Printf(TEXT("Mesh not found: %s"), *MeshPath));
+	}
+
+	FDynamicMesh3 DynMesh;
+	FString Error;
+	if (!MeshHelpers::LoadStaticMeshIntoDynamicMesh(SourceMesh, DynMesh, &Error))
+	{
+		return FECACommandResult::Error(Error);
+	}
+
+	// Compute normals first
+	FMeshNormals Normals(&DynMesh);
+	Normals.ComputeVertexNormals();
+
+	// Displace vertices
+	for (int32 VertID : DynMesh.VertexIndicesItr())
+	{
+		FVector3d Pos = DynMesh.GetVertex(VertID);
+		FVector3d Normal = Normals[VertID];
+		
+		double DisplaceAmount = Amount;
+		if (Mode == TEXT("random"))
+		{
+			DisplaceAmount = Amount * FMath::FRandRange(-1.0, 1.0);
+		}
+		else if (Mode == TEXT("noise"))
+		{
+			// Simple noise based on position
+			double NoiseVal = FMath::PerlinNoise3D(FVector(Pos.X * 0.01, Pos.Y * 0.01, Pos.Z * 0.01));
+			DisplaceAmount = Amount * NoiseVal;
+		}
+		
+		DynMesh.SetVertex(VertID, Pos + Normal * DisplaceAmount);
+	}
+
+	FMeshNormals::QuickComputeVertexNormals(DynMesh);
+
+	UStaticMesh* OutputMesh = nullptr;
+	if (OutputPath == MeshPath)
+	{
+		MeshHelpers::SaveDynamicMeshToStaticMesh(DynMesh, SourceMesh);
+		OutputMesh = SourceMesh;
+	}
+	else
+	{
+		OutputMesh = MeshHelpers::CreateStaticMeshFromDynamicMesh(DynMesh, OutputPath, true, true, &Error);
+		if (!OutputMesh)
+		{
+			return FECACommandResult::Error(Error);
+		}
+	}
+
+	TSharedPtr<FJsonObject> Result = MeshHelpers::BuildMeshResultJson(OutputMesh);
+	Result->SetStringField(TEXT("mode"), Mode);
+	Result->SetNumberField(TEXT("amount"), Amount);
+	return FECACommandResult::Success(Result);
+}
+
+FECACommandResult FECACommand_MeshBend::Execute(const TSharedPtr<FJsonObject>& Params)
+{
+	FString MeshPath;
+	if (!GetStringParam(Params, TEXT("mesh_path"), MeshPath))
+	{
+		return FECACommandResult::Error(TEXT("Missing required parameter: mesh_path"));
+	}
+
+	FString OutputPath;
+	if (!GetStringParam(Params, TEXT("output_path"), OutputPath, false))
+	{
+		OutputPath = MeshPath;
+	}
+
+	double Angle = 45.0; // Bend angle in degrees
+	GetFloatParam(Params, TEXT("angle"), Angle, false);
+
+	// Bend axis: "x", "y", "z"
+	FString Axis = TEXT("z");
+	GetStringParam(Params, TEXT("axis"), Axis, false);
+	Axis = Axis.ToLower();
+
+	UStaticMesh* SourceMesh = MeshHelpers::LoadStaticMesh(MeshPath);
+	if (!SourceMesh)
+	{
+		return FECACommandResult::Error(FString::Printf(TEXT("Mesh not found: %s"), *MeshPath));
+	}
+
+	FDynamicMesh3 DynMesh;
+	FString Error;
+	if (!MeshHelpers::LoadStaticMeshIntoDynamicMesh(SourceMesh, DynMesh, &Error))
+	{
+		return FECACommandResult::Error(Error);
+	}
+
+	// Get mesh bounds to determine bend range
+	FAxisAlignedBox3d Bounds = DynMesh.GetBounds();
+	
+	// Determine primary axis for bending
+	int32 BendAxis = 2; // Z by default
+	int32 HeightAxis = 2;
+	if (Axis == TEXT("x"))
+	{
+		BendAxis = 0;
+		HeightAxis = 2;
+	}
+	else if (Axis == TEXT("y"))
+	{
+		BendAxis = 1;
+		HeightAxis = 2;
+	}
+	else
+	{
+		BendAxis = 2;
+		HeightAxis = 0;
+	}
+
+	double Height = Bounds.Max[HeightAxis] - Bounds.Min[HeightAxis];
+	double BendRadians = FMath::DegreesToRadians(Angle);
+
+	// Bend the mesh
+	for (int32 VertID : DynMesh.VertexIndicesItr())
+	{
+		FVector3d Pos = DynMesh.GetVertex(VertID);
+		
+		// Normalize height position (0 to 1)
+		double T = (Pos[HeightAxis] - Bounds.Min[HeightAxis]) / Height;
+		T = FMath::Clamp(T, 0.0, 1.0);
+		
+		// Calculate bend angle for this vertex
+		double VertAngle = BendRadians * T;
+		
+		// Apply circular bend
+		if (FMath::Abs(BendRadians) > 0.001)
+		{
+			double Radius = Height / BendRadians;
+			double OriginalHeight = Pos[HeightAxis] - Bounds.Min[HeightAxis];
+			
+			FVector3d NewPos = Pos;
+			if (BendAxis == 0) // Bend around X
+			{
+				NewPos.Y = Pos.Y + Radius * (1.0 - FMath::Cos(VertAngle));
+				NewPos.Z = Bounds.Min[HeightAxis] + Radius * FMath::Sin(VertAngle);
+			}
+			else if (BendAxis == 1) // Bend around Y
+			{
+				NewPos.X = Pos.X + Radius * (1.0 - FMath::Cos(VertAngle));
+				NewPos.Z = Bounds.Min[HeightAxis] + Radius * FMath::Sin(VertAngle);
+			}
+			else // Bend around Z
+			{
+				NewPos.Y = Pos.Y + Radius * (1.0 - FMath::Cos(VertAngle));
+				NewPos.X = Bounds.Min[HeightAxis] + Radius * FMath::Sin(VertAngle);
+			}
+			
+			DynMesh.SetVertex(VertID, NewPos);
+		}
+	}
+
+	FMeshNormals::QuickComputeVertexNormals(DynMesh);
+
+	UStaticMesh* OutputMesh = nullptr;
+	if (OutputPath == MeshPath)
+	{
+		MeshHelpers::SaveDynamicMeshToStaticMesh(DynMesh, SourceMesh);
+		OutputMesh = SourceMesh;
+	}
+	else
+	{
+		OutputMesh = MeshHelpers::CreateStaticMeshFromDynamicMesh(DynMesh, OutputPath, true, true, &Error);
+		if (!OutputMesh)
+		{
+			return FECACommandResult::Error(Error);
+		}
+	}
+
+	TSharedPtr<FJsonObject> Result = MeshHelpers::BuildMeshResultJson(OutputMesh);
+	Result->SetNumberField(TEXT("angle"), Angle);
+	Result->SetStringField(TEXT("axis"), Axis);
+	return FECACommandResult::Success(Result);
+}
+
+FECACommandResult FECACommand_MeshTwist::Execute(const TSharedPtr<FJsonObject>& Params)
+{
+	FString MeshPath;
+	if (!GetStringParam(Params, TEXT("mesh_path"), MeshPath))
+	{
+		return FECACommandResult::Error(TEXT("Missing required parameter: mesh_path"));
+	}
+
+	FString OutputPath;
+	if (!GetStringParam(Params, TEXT("output_path"), OutputPath, false))
+	{
+		OutputPath = MeshPath;
+	}
+
+	double Angle = 90.0; // Total twist angle in degrees
+	GetFloatParam(Params, TEXT("angle"), Angle, false);
+
+	// Twist axis: "x", "y", "z"
+	FString Axis = TEXT("z");
+	GetStringParam(Params, TEXT("axis"), Axis, false);
+	Axis = Axis.ToLower();
+
+	UStaticMesh* SourceMesh = MeshHelpers::LoadStaticMesh(MeshPath);
+	if (!SourceMesh)
+	{
+		return FECACommandResult::Error(FString::Printf(TEXT("Mesh not found: %s"), *MeshPath));
+	}
+
+	FDynamicMesh3 DynMesh;
+	FString Error;
+	if (!MeshHelpers::LoadStaticMeshIntoDynamicMesh(SourceMesh, DynMesh, &Error))
+	{
+		return FECACommandResult::Error(Error);
+	}
+
+	// Get mesh bounds
+	FAxisAlignedBox3d Bounds = DynMesh.GetBounds();
+	FVector3d Center = Bounds.Center();
+	
+	// Determine twist axis
+	int32 TwistAxis = 2; // Z by default
+	if (Axis == TEXT("x")) TwistAxis = 0;
+	else if (Axis == TEXT("y")) TwistAxis = 1;
+	else TwistAxis = 2;
+
+	double Height = Bounds.Max[TwistAxis] - Bounds.Min[TwistAxis];
+	double TwistRadians = FMath::DegreesToRadians(Angle);
+
+	// Twist the mesh around the specified axis
+	for (int32 VertID : DynMesh.VertexIndicesItr())
+	{
+		FVector3d Pos = DynMesh.GetVertex(VertID);
+		
+		// Normalize position along twist axis (0 to 1)
+		double T = (Pos[TwistAxis] - Bounds.Min[TwistAxis]) / Height;
+		T = FMath::Clamp(T, 0.0, 1.0);
+		
+		// Calculate twist angle for this vertex
+		double VertAngle = TwistRadians * T;
+		double CosA = FMath::Cos(VertAngle);
+		double SinA = FMath::Sin(VertAngle);
+		
+		FVector3d NewPos = Pos;
+		if (TwistAxis == 0) // Twist around X
+		{
+			double RelY = Pos.Y - Center.Y;
+			double RelZ = Pos.Z - Center.Z;
+			NewPos.Y = Center.Y + RelY * CosA - RelZ * SinA;
+			NewPos.Z = Center.Z + RelY * SinA + RelZ * CosA;
+		}
+		else if (TwistAxis == 1) // Twist around Y
+		{
+			double RelX = Pos.X - Center.X;
+			double RelZ = Pos.Z - Center.Z;
+			NewPos.X = Center.X + RelX * CosA - RelZ * SinA;
+			NewPos.Z = Center.Z + RelX * SinA + RelZ * CosA;
+		}
+		else // Twist around Z
+		{
+			double RelX = Pos.X - Center.X;
+			double RelY = Pos.Y - Center.Y;
+			NewPos.X = Center.X + RelX * CosA - RelY * SinA;
+			NewPos.Y = Center.Y + RelX * SinA + RelY * CosA;
+		}
+		
+		DynMesh.SetVertex(VertID, NewPos);
+	}
+
+	FMeshNormals::QuickComputeVertexNormals(DynMesh);
+
+	UStaticMesh* OutputMesh = nullptr;
+	if (OutputPath == MeshPath)
+	{
+		MeshHelpers::SaveDynamicMeshToStaticMesh(DynMesh, SourceMesh);
+		OutputMesh = SourceMesh;
+	}
+	else
+	{
+		OutputMesh = MeshHelpers::CreateStaticMeshFromDynamicMesh(DynMesh, OutputPath, true, true, &Error);
+		if (!OutputMesh)
+		{
+			return FECACommandResult::Error(Error);
+		}
+	}
+
+	TSharedPtr<FJsonObject> Result = MeshHelpers::BuildMeshResultJson(OutputMesh);
+	Result->SetNumberField(TEXT("angle"), Angle);
+	Result->SetStringField(TEXT("axis"), Axis);
+	return FECACommandResult::Success(Result);
+}
+
+//==============================================================================
+// HEIGHTFIELD
+//==============================================================================
+
+FECACommandResult FECACommand_CreateHeightfield::Execute(const TSharedPtr<FJsonObject>& Params)
+{
+	FString AssetPath;
+	if (!GetStringParam(Params, TEXT("asset_path"), AssetPath))
+	{
+		return FECACommandResult::Error(TEXT("Missing required parameter: asset_path"));
+	}
+
+	double SizeX = 1000, SizeY = 1000, HeightScale = 100;
+	GetFloatParam(Params, TEXT("size_x"), SizeX, false);
+	GetFloatParam(Params, TEXT("size_y"), SizeY, false);
+	GetFloatParam(Params, TEXT("height_scale"), HeightScale, false);
+
+	int32 SubdivisionsX = 64, SubdivisionsY = 64;
+	GetIntParam(Params, TEXT("subdivisions_x"), SubdivisionsX, false);
+	GetIntParam(Params, TEXT("subdivisions_y"), SubdivisionsY, false);
+
+	// Create heightfield grid
+	FRectangleMeshGenerator GridGen;
+	GridGen.Width = SizeX;
+	GridGen.Height = SizeY;
+	GridGen.WidthVertexCount = SubdivisionsX + 1;
+	GridGen.HeightVertexCount = SubdivisionsY + 1;
+	GridGen.bSinglePolyGroup = true;
+
+	GridGen.Generate();
+
+	FDynamicMesh3 DynMesh;
+	DynMesh.Copy(&GridGen);
+
+	// Apply heights from array if provided
+	const TArray<TSharedPtr<FJsonValue>>* HeightsArray = nullptr;
+	if (GetArrayParam(Params, TEXT("heights"), HeightsArray, false) && HeightsArray->Num() > 0)
+	{
+		int32 VertIdx = 0;
+		for (int32 y = 0; y < SubdivisionsY + 1 && y < HeightsArray->Num(); y++)
+		{
+			const TArray<TSharedPtr<FJsonValue>>* RowArray = nullptr;
+			if ((*HeightsArray)[y]->TryGetArray(RowArray))
+			{
+				for (int32 x = 0; x < SubdivisionsX + 1 && x < RowArray->Num(); x++)
+				{
+					if (VertIdx < DynMesh.VertexCount())
+					{
+						FVector3d Pos = DynMesh.GetVertex(VertIdx);
+						Pos.Z = (*RowArray)[x]->AsNumber() * HeightScale;
+						DynMesh.SetVertex(VertIdx, Pos);
+					}
+					VertIdx++;
+				}
+			}
+		}
+	}
+
+	FMeshNormals::QuickComputeVertexNormals(DynMesh);
+
+	FString Error;
+	UStaticMesh* StaticMesh = MeshHelpers::CreateStaticMeshFromDynamicMesh(DynMesh, AssetPath, true, true, &Error);
+	if (!StaticMesh)
+	{
+		return FECACommandResult::Error(Error);
+	}
+
+	TSharedPtr<FJsonObject> Result = MeshHelpers::BuildMeshResultJson(StaticMesh, DynMesh.VertexCount(), DynMesh.TriangleCount());
+	return FECACommandResult::Success(Result);
+}
+
+//==============================================================================
+// UV AND NORMALS
+//==============================================================================
+
+FECACommandResult FECACommand_MeshGenerateUVs::Execute(const TSharedPtr<FJsonObject>& Params)
+{
+	FString MeshPath;
+	if (!GetStringParam(Params, TEXT("mesh_path"), MeshPath))
+	{
+		return FECACommandResult::Error(TEXT("Missing required parameter: mesh_path"));
+	}
+
+	FString OutputPath;
+	if (!GetStringParam(Params, TEXT("output_path"), OutputPath, false))
+	{
+		OutputPath = MeshPath;
+	}
+
+	// UV generation method: "box", "planar", "cylindrical", "spherical"
+	FString Method = TEXT("box");
+	GetStringParam(Params, TEXT("method"), Method, false);
+	Method = Method.ToLower();
+
+	double Scale = 1.0;
+	GetFloatParam(Params, TEXT("scale"), Scale, false);
+
+	UStaticMesh* SourceMesh = MeshHelpers::LoadStaticMesh(MeshPath);
+	if (!SourceMesh)
+	{
+		return FECACommandResult::Error(FString::Printf(TEXT("Mesh not found: %s"), *MeshPath));
+	}
+
+	FDynamicMesh3 DynMesh;
+	FString Error;
+	if (!MeshHelpers::LoadStaticMeshIntoDynamicMesh(SourceMesh, DynMesh, &Error))
+	{
+		return FECACommandResult::Error(Error);
+	}
+
+	// Ensure we have UV attributes
+	if (!DynMesh.HasAttributes())
+	{
+		DynMesh.EnableAttributes();
+	}
+	DynMesh.Attributes()->SetNumUVLayers(1);
+	FDynamicMeshUVOverlay* UVOverlay = DynMesh.Attributes()->PrimaryUV();
+
+	// Get mesh bounds for UV calculations
+	FAxisAlignedBox3d Bounds = DynMesh.GetBounds();
+	FVector3d Size = Bounds.Diagonal();
+	FVector3d Center = Bounds.Center();
+
+	// Clear existing UVs
+	UVOverlay->ClearElements();
+
+	if (Method == TEXT("box") || Method == TEXT("cube"))
+	{
+		// Box projection - project UVs based on face normal direction
+		for (int32 TriID : DynMesh.TriangleIndicesItr())
+		{
+			FIndex3i Tri = DynMesh.GetTriangle(TriID);
+			FVector3d V0 = DynMesh.GetVertex(Tri.A);
+			FVector3d V1 = DynMesh.GetVertex(Tri.B);
+			FVector3d V2 = DynMesh.GetVertex(Tri.C);
+			
+			// Compute face normal
+			FVector3d Normal = ((V1 - V0).Cross(V2 - V0)).GetSafeNormal();
+			FVector3d AbsNormal(FMath::Abs(Normal.X), FMath::Abs(Normal.Y), FMath::Abs(Normal.Z));
+			
+			// Determine dominant axis
+			FVector2f UV0, UV1, UV2;
+			if (AbsNormal.X >= AbsNormal.Y && AbsNormal.X >= AbsNormal.Z)
+			{
+				// Project onto YZ plane
+				UV0 = FVector2f((V0.Y - Bounds.Min.Y) / Size.Y, (V0.Z - Bounds.Min.Z) / Size.Z) * Scale;
+				UV1 = FVector2f((V1.Y - Bounds.Min.Y) / Size.Y, (V1.Z - Bounds.Min.Z) / Size.Z) * Scale;
+				UV2 = FVector2f((V2.Y - Bounds.Min.Y) / Size.Y, (V2.Z - Bounds.Min.Z) / Size.Z) * Scale;
+			}
+			else if (AbsNormal.Y >= AbsNormal.X && AbsNormal.Y >= AbsNormal.Z)
+			{
+				// Project onto XZ plane
+				UV0 = FVector2f((V0.X - Bounds.Min.X) / Size.X, (V0.Z - Bounds.Min.Z) / Size.Z) * Scale;
+				UV1 = FVector2f((V1.X - Bounds.Min.X) / Size.X, (V1.Z - Bounds.Min.Z) / Size.Z) * Scale;
+				UV2 = FVector2f((V2.X - Bounds.Min.X) / Size.X, (V2.Z - Bounds.Min.Z) / Size.Z) * Scale;
+			}
+			else
+			{
+				// Project onto XY plane
+				UV0 = FVector2f((V0.X - Bounds.Min.X) / Size.X, (V0.Y - Bounds.Min.Y) / Size.Y) * Scale;
+				UV1 = FVector2f((V1.X - Bounds.Min.X) / Size.X, (V1.Y - Bounds.Min.Y) / Size.Y) * Scale;
+				UV2 = FVector2f((V2.X - Bounds.Min.X) / Size.X, (V2.Y - Bounds.Min.Y) / Size.Y) * Scale;
+			}
+			
+			int32 UVID0 = UVOverlay->AppendElement(UV0);
+			int32 UVID1 = UVOverlay->AppendElement(UV1);
+			int32 UVID2 = UVOverlay->AppendElement(UV2);
+			UVOverlay->SetTriangle(TriID, FIndex3i(UVID0, UVID1, UVID2));
+		}
+	}
+	else if (Method == TEXT("planar"))
+	{
+		// Planar projection from top (XY plane)
+		for (int32 TriID : DynMesh.TriangleIndicesItr())
+		{
+			FIndex3i Tri = DynMesh.GetTriangle(TriID);
+			FVector3d V0 = DynMesh.GetVertex(Tri.A);
+			FVector3d V1 = DynMesh.GetVertex(Tri.B);
+			FVector3d V2 = DynMesh.GetVertex(Tri.C);
+			
+			FVector2f UV0((V0.X - Bounds.Min.X) / Size.X * Scale, (V0.Y - Bounds.Min.Y) / Size.Y * Scale);
+			FVector2f UV1((V1.X - Bounds.Min.X) / Size.X * Scale, (V1.Y - Bounds.Min.Y) / Size.Y * Scale);
+			FVector2f UV2((V2.X - Bounds.Min.X) / Size.X * Scale, (V2.Y - Bounds.Min.Y) / Size.Y * Scale);
+			
+			int32 UVID0 = UVOverlay->AppendElement(UV0);
+			int32 UVID1 = UVOverlay->AppendElement(UV1);
+			int32 UVID2 = UVOverlay->AppendElement(UV2);
+			UVOverlay->SetTriangle(TriID, FIndex3i(UVID0, UVID1, UVID2));
+		}
+	}
+	else if (Method == TEXT("cylindrical"))
+	{
+		// Cylindrical projection around Z axis
+		for (int32 TriID : DynMesh.TriangleIndicesItr())
+		{
+			FIndex3i Tri = DynMesh.GetTriangle(TriID);
+			FVector3d V0 = DynMesh.GetVertex(Tri.A);
+			FVector3d V1 = DynMesh.GetVertex(Tri.B);
+			FVector3d V2 = DynMesh.GetVertex(Tri.C);
+			
+			auto ComputeCylindricalUV = [&](const FVector3d& V) -> FVector2f
+			{
+				double Angle = FMath::Atan2(V.Y - Center.Y, V.X - Center.X);
+				double U = (Angle + PI) / (2.0 * PI);
+				double VCoord = (V.Z - Bounds.Min.Z) / Size.Z;
+				return FVector2f(U * Scale, VCoord * Scale);
+			};
+			
+			FVector2f UV0 = ComputeCylindricalUV(V0);
+			FVector2f UV1 = ComputeCylindricalUV(V1);
+			FVector2f UV2 = ComputeCylindricalUV(V2);
+			
+			int32 UVID0 = UVOverlay->AppendElement(UV0);
+			int32 UVID1 = UVOverlay->AppendElement(UV1);
+			int32 UVID2 = UVOverlay->AppendElement(UV2);
+			UVOverlay->SetTriangle(TriID, FIndex3i(UVID0, UVID1, UVID2));
+		}
+	}
+	else if (Method == TEXT("spherical"))
+	{
+		// Spherical projection from center
+		for (int32 TriID : DynMesh.TriangleIndicesItr())
+		{
+			FIndex3i Tri = DynMesh.GetTriangle(TriID);
+			FVector3d V0 = DynMesh.GetVertex(Tri.A);
+			FVector3d V1 = DynMesh.GetVertex(Tri.B);
+			FVector3d V2 = DynMesh.GetVertex(Tri.C);
+			
+			auto ComputeSphericalUV = [&](const FVector3d& V) -> FVector2f
+			{
+				FVector3d Dir = (V - Center).GetSafeNormal();
+				double U = 0.5 + FMath::Atan2(Dir.Y, Dir.X) / (2.0 * PI);
+				double VCoord = 0.5 + FMath::Asin(FMath::Clamp(Dir.Z, -1.0, 1.0)) / PI;
+				return FVector2f(U * Scale, VCoord * Scale);
+			};
+			
+			FVector2f UV0 = ComputeSphericalUV(V0);
+			FVector2f UV1 = ComputeSphericalUV(V1);
+			FVector2f UV2 = ComputeSphericalUV(V2);
+			
+			int32 UVID0 = UVOverlay->AppendElement(UV0);
+			int32 UVID1 = UVOverlay->AppendElement(UV1);
+			int32 UVID2 = UVOverlay->AppendElement(UV2);
+			UVOverlay->SetTriangle(TriID, FIndex3i(UVID0, UVID1, UVID2));
+		}
+	}
+	else
+	{
+		return FECACommandResult::Error(FString::Printf(TEXT("Unknown UV method: %s. Use 'box', 'planar', 'cylindrical', or 'spherical'."), *Method));
+	}
+
+	UStaticMesh* OutputMesh = nullptr;
+	if (OutputPath == MeshPath)
+	{
+		MeshHelpers::SaveDynamicMeshToStaticMesh(DynMesh, SourceMesh);
+		OutputMesh = SourceMesh;
+	}
+	else
+	{
+		OutputMesh = MeshHelpers::CreateStaticMeshFromDynamicMesh(DynMesh, OutputPath, true, true, &Error);
+		if (!OutputMesh)
+		{
+			return FECACommandResult::Error(Error);
+		}
+	}
+
+	TSharedPtr<FJsonObject> Result = MeshHelpers::BuildMeshResultJson(OutputMesh);
+	Result->SetStringField(TEXT("uv_method"), Method);
+	Result->SetNumberField(TEXT("uv_scale"), Scale);
+	return FECACommandResult::Success(Result);
+}
+
+FECACommandResult FECACommand_MeshRecomputeNormals::Execute(const TSharedPtr<FJsonObject>& Params)
+{
+	FString MeshPath;
+	if (!GetStringParam(Params, TEXT("mesh_path"), MeshPath))
+	{
+		return FECACommandResult::Error(TEXT("Missing required parameter: mesh_path"));
+	}
+
+	FString OutputPath;
+	if (!GetStringParam(Params, TEXT("output_path"), OutputPath, false))
+	{
+		OutputPath = MeshPath;
+	}
+	
+	// Get normal computation options
+	bool bWeldSeams = true; // Default to seam-aware normals
+	GetBoolParam(Params, TEXT("weld_seams"), bWeldSeams, false);
+	
+	double WeldTolerance = 0.0001;
+	GetFloatParam(Params, TEXT("weld_tolerance"), WeldTolerance, false);
+	
+	bool bWeightByArea = true;
+	GetBoolParam(Params, TEXT("weight_by_area"), bWeightByArea, false);
+	
+	bool bWeightByAngle = true;
+	GetBoolParam(Params, TEXT("weight_by_angle"), bWeightByAngle, false);
+
+	UStaticMesh* SourceMesh = MeshHelpers::LoadStaticMesh(MeshPath);
+	if (!SourceMesh)
+	{
+		return FECACommandResult::Error(FString::Printf(TEXT("Mesh not found: %s"), *MeshPath));
+	}
+
+	FDynamicMesh3 DynMesh;
+	FString Error;
+	if (!MeshHelpers::LoadStaticMeshIntoDynamicMesh(SourceMesh, DynMesh, &Error))
+	{
+		return FECACommandResult::Error(Error);
+	}
+
+	// Use seam-aware normal computation by default to properly smooth across UV seams
+	if (bWeldSeams)
+	{
+		MeshHelpers::ComputeVertexNormalsAcrossSeams(DynMesh, WeldTolerance, bWeightByArea, bWeightByAngle);
+	}
+	else
+	{
+		// Use quick computation (per-vertex-index, doesn't weld seams)
+		FMeshNormals::QuickComputeVertexNormals(DynMesh);
+	}
+
+	UStaticMesh* OutputMesh = nullptr;
+	if (OutputPath == MeshPath)
+	{
+		MeshHelpers::SaveDynamicMeshToStaticMesh(DynMesh, SourceMesh);
+		OutputMesh = SourceMesh;
+	}
+	else
+	{
+		OutputMesh = MeshHelpers::CreateStaticMeshFromDynamicMesh(DynMesh, OutputPath, true, true, &Error);
+		if (!OutputMesh)
+		{
+			return FECACommandResult::Error(Error);
+		}
+	}
+
+	TSharedPtr<FJsonObject> Result = MeshHelpers::BuildMeshResultJson(OutputMesh);
+	return FECACommandResult::Success(Result);
+}
+
+FECACommandResult FECACommand_MeshFlipNormals::Execute(const TSharedPtr<FJsonObject>& Params)
+{
+	FString MeshPath;
+	if (!GetStringParam(Params, TEXT("mesh_path"), MeshPath))
+	{
+		return FECACommandResult::Error(TEXT("Missing required parameter: mesh_path"));
+	}
+
+	FString OutputPath;
+	if (!GetStringParam(Params, TEXT("output_path"), OutputPath, false))
+	{
+		OutputPath = MeshPath;
+	}
+
+	UStaticMesh* SourceMesh = MeshHelpers::LoadStaticMesh(MeshPath);
+	if (!SourceMesh)
+	{
+		return FECACommandResult::Error(FString::Printf(TEXT("Mesh not found: %s"), *MeshPath));
+	}
+
+	FDynamicMesh3 DynMesh;
+	FString Error;
+	if (!MeshHelpers::LoadStaticMeshIntoDynamicMesh(SourceMesh, DynMesh, &Error))
+	{
+		return FECACommandResult::Error(Error);
+	}
+
+	// Flip vertex normals only, don't change triangle winding
+	if (DynMesh.HasAttributes() && DynMesh.Attributes()->PrimaryNormals())
+	{
+		FDynamicMeshNormalOverlay* NormalOverlay = DynMesh.Attributes()->PrimaryNormals();
+		for (int32 ElementID : NormalOverlay->ElementIndicesItr())
+		{
+			FVector3f Normal = NormalOverlay->GetElement(ElementID);
+			NormalOverlay->SetElement(ElementID, -Normal);
+		}
+	}
+	else
+	{
+		// Recompute normals first if not present, then flip them
+		FMeshNormals::QuickComputeVertexNormals(DynMesh);
+		if (DynMesh.HasAttributes() && DynMesh.Attributes()->PrimaryNormals())
+		{
+			FDynamicMeshNormalOverlay* NormalOverlay = DynMesh.Attributes()->PrimaryNormals();
+			for (int32 ElementID : NormalOverlay->ElementIndicesItr())
+			{
+				FVector3f Normal = NormalOverlay->GetElement(ElementID);
+				NormalOverlay->SetElement(ElementID, -Normal);
+			}
+		}
+	}
+
+	UStaticMesh* OutputMesh = nullptr;
+	if (OutputPath == MeshPath)
+	{
+		MeshHelpers::SaveDynamicMeshToStaticMesh(DynMesh, SourceMesh);
+		OutputMesh = SourceMesh;
+	}
+	else
+	{
+		OutputMesh = MeshHelpers::CreateStaticMeshFromDynamicMesh(DynMesh, OutputPath, true, true, &Error);
+		if (!OutputMesh)
+		{
+			return FECACommandResult::Error(Error);
+		}
+	}
+
+	TSharedPtr<FJsonObject> Result = MeshHelpers::BuildMeshResultJson(OutputMesh);
+	return FECACommandResult::Success(Result);
+}
+
+
+FECACommandResult FECACommand_MeshReverseWinding::Execute(const TSharedPtr<FJsonObject>& Params)
+{
+	FString MeshPath;
+	if (!GetStringParam(Params, TEXT("mesh_path"), MeshPath))
+	{
+		return FECACommandResult::Error(TEXT("Missing required parameter: mesh_path"));
+	}
+
+	FString OutputPath;
+	if (!GetStringParam(Params, TEXT("output_path"), OutputPath, false))
+	{
+		OutputPath = MeshPath;
+	}
+
+	UStaticMesh* SourceMesh = MeshHelpers::LoadStaticMesh(MeshPath);
+	if (!SourceMesh)
+	{
+		return FECACommandResult::Error(FString::Printf(TEXT("Mesh not found: %s"), *MeshPath));
+	}
+
+	FDynamicMesh3 DynMesh;
+	FString Error;
+	if (!MeshHelpers::LoadStaticMeshIntoDynamicMesh(SourceMesh, DynMesh, &Error))
+	{
+		return FECACommandResult::Error(Error);
+	}
+
+	// Only reverse triangle winding order, preserve existing normals
+	DynMesh.ReverseOrientation(false);  // false = don't flip normals
+
+	UStaticMesh* OutputMesh = nullptr;
+	if (OutputPath == MeshPath)
+	{
+		MeshHelpers::SaveDynamicMeshToStaticMesh(DynMesh, SourceMesh);
+		OutputMesh = SourceMesh;
+	}
+	else
+	{
+		OutputMesh = MeshHelpers::CreateStaticMeshFromDynamicMesh(DynMesh, OutputPath, true, true, &Error);
+		if (!OutputMesh)
+		{
+			return FECACommandResult::Error(Error);
+		}
+	}
+
+	TSharedPtr<FJsonObject> Result = MeshHelpers::BuildMeshResultJson(OutputMesh);
+	return FECACommandResult::Success(Result);
+}
+
+FECACommandResult FECACommand_MeshRepair::Execute(const TSharedPtr<FJsonObject>& Params)
+{
+	FString MeshPath;
+	if (!GetStringParam(Params, TEXT("mesh_path"), MeshPath))
+	{
+		return FECACommandResult::Error(TEXT("Missing required parameter: mesh_path"));
+	}
+
+	FString OutputPath;
+	if (!GetStringParam(Params, TEXT("output_path"), OutputPath, false))
+	{
+		OutputPath = MeshPath;
+	}
+
+	bool bFillHoles = true;
+	GetBoolParam(Params, TEXT("fill_holes"), bFillHoles, false);
+
+	bool bRemoveDegenerates = true;
+	GetBoolParam(Params, TEXT("remove_degenerates"), bRemoveDegenerates, false);
+
+	bool bFixNormals = true;
+	GetBoolParam(Params, TEXT("fix_normals"), bFixNormals, false);
+
+	double WeldTolerance = 0.001;
+	GetFloatParam(Params, TEXT("weld_tolerance"), WeldTolerance, false);
+
+	UStaticMesh* SourceMesh = MeshHelpers::LoadStaticMesh(MeshPath);
+	if (!SourceMesh)
+	{
+		return FECACommandResult::Error(FString::Printf(TEXT("Mesh not found: %s"), *MeshPath));
+	}
+
+	FDynamicMesh3 DynMesh;
+	FString Error;
+	if (!MeshHelpers::LoadStaticMeshIntoDynamicMesh(SourceMesh, DynMesh, &Error))
+	{
+		return FECACommandResult::Error(Error);
+	}
+
+	int32 OriginalVertCount = DynMesh.VertexCount();
+	int32 OriginalTriCount = DynMesh.TriangleCount();
+	int32 DegeneratesRemoved = 0;
+	int32 HolesFilled = 0;
+
+	// Remove degenerate triangles (zero area)
+	if (bRemoveDegenerates)
+	{
+		TArray<int32> DegenerateTriangles;
+		for (int32 TriID : DynMesh.TriangleIndicesItr())
+		{
+			FIndex3i Tri = DynMesh.GetTriangle(TriID);
+			FVector3d V0 = DynMesh.GetVertex(Tri.A);
+			FVector3d V1 = DynMesh.GetVertex(Tri.B);
+			FVector3d V2 = DynMesh.GetVertex(Tri.C);
+			
+			// Check for degenerate triangle (very small area)
+			FVector3d Cross = (V1 - V0).Cross(V2 - V0);
+			double Area = Cross.Length() * 0.5;
+			
+			if (Area < 0.0001 || Tri.A == Tri.B || Tri.B == Tri.C || Tri.A == Tri.C)
+			{
+				DegenerateTriangles.Add(TriID);
+			}
+		}
+		
+		for (int32 TriID : DegenerateTriangles)
+		{
+			if (DynMesh.IsTriangle(TriID))
+			{
+				DynMesh.RemoveTriangle(TriID);
+				DegeneratesRemoved++;
+			}
+		}
+	}
+
+	// Find and fill boundary holes
+	if (bFillHoles)
+	{
+		// Find boundary edges (edges with only one triangle)
+		TArray<int32> BoundaryEdges;
+		for (int32 EdgeID : DynMesh.EdgeIndicesItr())
+		{
+			FIndex2i EdgeTris = DynMesh.GetEdgeT(EdgeID);
+			if (EdgeTris.B == FDynamicMesh3::InvalidID)
+			{
+				BoundaryEdges.Add(EdgeID);
+			}
+		}
+		
+		// Simple hole filling: try to close small boundary loops
+		// This is a simplified implementation
+		TSet<int32> ProcessedEdges;
+		for (int32 StartEdge : BoundaryEdges)
+		{
+			if (ProcessedEdges.Contains(StartEdge)) continue;
+			
+			// Try to trace a boundary loop
+			TArray<int32> LoopVertices;
+			int32 CurrentEdge = StartEdge;
+			int32 MaxIterations = 1000;
+			int32 Iterations = 0;
+			
+			FIndex2i FirstEdgeV = DynMesh.GetEdgeV(StartEdge);
+			int32 CurrentVertex = FirstEdgeV.A;
+			LoopVertices.Add(CurrentVertex);
+			
+			while (Iterations++ < MaxIterations)
+			{
+				ProcessedEdges.Add(CurrentEdge);
+				
+				FIndex2i EdgeV = DynMesh.GetEdgeV(CurrentEdge);
+				int32 NextVertex = (EdgeV.A == CurrentVertex) ? EdgeV.B : EdgeV.A;
+				
+				if (NextVertex == LoopVertices[0] && LoopVertices.Num() >= 3)
+				{
+					// Closed the loop - fill it with a fan triangulation (reversed winding)
+					for (int32 i = 1; i < LoopVertices.Num() - 1; i++)
+					{
+						DynMesh.AppendTriangle(LoopVertices[0], LoopVertices[i + 1], LoopVertices[i]);
+					}
+					HolesFilled++;
+					break;
+				}
+				
+				LoopVertices.Add(NextVertex);
+				CurrentVertex = NextVertex;
+				
+				// Find next boundary edge from current vertex
+				bool Found = false;
+				for (int32 EdgeID : DynMesh.VtxEdgesItr(CurrentVertex))
+				{
+					if (EdgeID != CurrentEdge && !ProcessedEdges.Contains(EdgeID))
+					{
+						FIndex2i EdgeTris = DynMesh.GetEdgeT(EdgeID);
+						if (EdgeTris.B == FDynamicMesh3::InvalidID)
+						{
+							CurrentEdge = EdgeID;
+							Found = true;
+							break;
+						}
+					}
+				}
+				
+				if (!Found) break;
+			}
+		}
+	}
+
+	// Recompute normals
+	if (bFixNormals)
+	{
+		FMeshNormals::QuickComputeVertexNormals(DynMesh);
+	}
+
+	UStaticMesh* OutputMesh = nullptr;
+	if (OutputPath == MeshPath)
+	{
+		MeshHelpers::SaveDynamicMeshToStaticMesh(DynMesh, SourceMesh);
+		OutputMesh = SourceMesh;
+	}
+	else
+	{
+		OutputMesh = MeshHelpers::CreateStaticMeshFromDynamicMesh(DynMesh, OutputPath, true, true, &Error);
+		if (!OutputMesh)
+		{
+			return FECACommandResult::Error(Error);
+		}
+	}
+
+	TSharedPtr<FJsonObject> Result = MeshHelpers::BuildMeshResultJson(OutputMesh);
+	Result->SetNumberField(TEXT("original_vertices"), OriginalVertCount);
+	Result->SetNumberField(TEXT("original_triangles"), OriginalTriCount);
+	Result->SetNumberField(TEXT("degenerates_removed"), DegeneratesRemoved);
+	Result->SetNumberField(TEXT("holes_filled"), HolesFilled);
+	return FECACommandResult::Success(Result);
+}
+
+//==============================================================================
+// MATERIALS
+//==============================================================================
+
+FECACommandResult FECACommand_MeshSetMaterials::Execute(const TSharedPtr<FJsonObject>& Params)
+{
+	FString MeshPath;
+	if (!GetStringParam(Params, TEXT("mesh_path"), MeshPath))
+	{
+		return FECACommandResult::Error(TEXT("Missing required parameter: mesh_path"));
+	}
+
+	const TArray<TSharedPtr<FJsonValue>>* MaterialSlotsArray;
+	if (!GetArrayParam(Params, TEXT("material_slots"), MaterialSlotsArray))
+	{
+		return FECACommandResult::Error(TEXT("Missing required parameter: material_slots"));
+	}
+
+	// Optional per-triangle material indices
+	const TArray<TSharedPtr<FJsonValue>>* TriangleMaterialsArray = nullptr;
+	GetArrayParam(Params, TEXT("triangle_materials"), TriangleMaterialsArray, false);
+
+	// Load the mesh
+	UStaticMesh* StaticMesh = MeshHelpers::LoadStaticMesh(MeshPath);
+	if (!StaticMesh)
+	{
+		return FECACommandResult::Error(FString::Printf(TEXT("Mesh not found: %s"), *MeshPath));
+	}
+
+	// Parse material slots
+	struct FMaterialSlotDef
+	{
+		FString Name;
+		FString MaterialPath;
+	};
+	TArray<FMaterialSlotDef> SlotDefs;
+
+	for (const TSharedPtr<FJsonValue>& SlotValue : *MaterialSlotsArray)
+	{
+		FMaterialSlotDef Def;
+		
+		// Support simple string (just name) or object with name and material_path
+		FString SlotStr;
+		if (SlotValue->TryGetString(SlotStr))
+		{
+			Def.Name = SlotStr;
+		}
+		else
+		{
+			const TSharedPtr<FJsonObject>* SlotObj;
+			if (SlotValue->TryGetObject(SlotObj))
+			{
+				(*SlotObj)->TryGetStringField(TEXT("name"), Def.Name);
+				(*SlotObj)->TryGetStringField(TEXT("material_path"), Def.MaterialPath);
+				
+				if (Def.Name.IsEmpty())
+				{
+					Def.Name = FString::Printf(TEXT("Material_%d"), SlotDefs.Num());
+				}
+			}
+		}
+		
+		SlotDefs.Add(Def);
+	}
+
+	if (SlotDefs.Num() == 0)
+	{
+		return FECACommandResult::Error(TEXT("At least one material slot must be specified"));
+	}
+
+	// Update material slots
+	StaticMesh->GetStaticMaterials().Empty();
+	for (int32 i = 0; i < SlotDefs.Num(); i++)
+	{
+		UMaterialInterface* Material = nullptr;
+		if (!SlotDefs[i].MaterialPath.IsEmpty())
+		{
+			Material = LoadObject<UMaterialInterface>(nullptr, *SlotDefs[i].MaterialPath);
+		}
+		
+		FName SlotName(*SlotDefs[i].Name);
+		FStaticMaterial NewMaterial(Material, SlotName, SlotName);
+		NewMaterial.UVChannelData = FMeshUVChannelInfo(1.0f);
+		StaticMesh->GetStaticMaterials().Add(NewMaterial);
+	}
+
+	// If per-triangle materials are provided, we need to rebuild the mesh with proper material IDs
+	if (TriangleMaterialsArray && TriangleMaterialsArray->Num() > 0)
+	{
+		FDynamicMesh3 DynMesh;
+		FString Error;
+		if (!MeshHelpers::LoadStaticMeshIntoDynamicMesh(StaticMesh, DynMesh, &Error))
+		{
+			return FECACommandResult::Error(Error);
+		}
+
+		// Enable material ID attribute
+		DynMesh.Attributes()->EnableMaterialID();
+		FDynamicMeshMaterialAttribute* MaterialAttr = DynMesh.Attributes()->GetMaterialID();
+
+		// Apply per-triangle material indices
+		int32 TriIndex = 0;
+		for (int32 TriID : DynMesh.TriangleIndicesItr())
+		{
+			if (TriIndex < TriangleMaterialsArray->Num())
+			{
+				int32 MatIndex = (*TriangleMaterialsArray)[TriIndex]->AsNumber();
+				MatIndex = FMath::Clamp(MatIndex, 0, SlotDefs.Num() - 1);
+				MaterialAttr->SetValue(TriID, MatIndex);
+			}
+			TriIndex++;
+		}
+
+		// Save back to static mesh
+		MeshHelpers::SaveDynamicMeshToStaticMesh(DynMesh, StaticMesh);
+	}
+
+	// Rebuild the mesh
+	StaticMesh->Build(false);
+	StaticMesh->PostEditChange();
+	StaticMesh->MarkPackageDirty();
+
+	TSharedPtr<FJsonObject> Result = MakeShareable(new FJsonObject());
+	Result->SetStringField(TEXT("mesh_path"), MeshPath);
+	Result->SetNumberField(TEXT("material_slots"), SlotDefs.Num());
+	
+	TArray<TSharedPtr<FJsonValue>> SlotsJson;
+	for (int32 i = 0; i < SlotDefs.Num(); i++)
+	{
+		TSharedPtr<FJsonObject> SlotJson = MakeShareable(new FJsonObject());
+		SlotJson->SetNumberField(TEXT("index"), i);
+		SlotJson->SetStringField(TEXT("name"), SlotDefs[i].Name);
+		if (!SlotDefs[i].MaterialPath.IsEmpty())
+		{
+			SlotJson->SetStringField(TEXT("material_path"), SlotDefs[i].MaterialPath);
+		}
+		SlotsJson.Add(MakeShareable(new FJsonValueObject(SlotJson)));
+	}
+	Result->SetArrayField(TEXT("slots"), SlotsJson);
+
+	return FECACommandResult::Success(Result);
+}
+
+//==============================================================================
+// TRANSFORM
+//==============================================================================
+
+FECACommandResult FECACommand_MeshTransform::Execute(const TSharedPtr<FJsonObject>& Params)
+{
+	FString MeshPath;
+	if (!GetStringParam(Params, TEXT("mesh_path"), MeshPath))
+	{
+		return FECACommandResult::Error(TEXT("Missing required parameter: mesh_path"));
+	}
+
+	FString OutputPath;
+	if (!GetStringParam(Params, TEXT("output_path"), OutputPath, false))
+	{
+		OutputPath = MeshPath;
+	}
+
+	UStaticMesh* SourceMesh = MeshHelpers::LoadStaticMesh(MeshPath);
+	if (!SourceMesh)
+	{
+		return FECACommandResult::Error(FString::Printf(TEXT("Mesh not found: %s"), *MeshPath));
+	}
+
+	FDynamicMesh3 DynMesh;
+	FString Error;
+	if (!MeshHelpers::LoadStaticMeshIntoDynamicMesh(SourceMesh, DynMesh, &Error))
+	{
+		return FECACommandResult::Error(Error);
+	}
+
+	// Get pivot point (defaults to mesh center)
+	FVector3d Pivot = FVector3d::Zero();
+	const TSharedPtr<FJsonObject>* PivotObj = nullptr;
+	if (Params->TryGetObjectField(TEXT("pivot"), PivotObj))
+	{
+		Pivot = MeshHelpers::ParseVector(*PivotObj);
+	}
+	else
+	{
+		// Default pivot to mesh center
+		FAxisAlignedBox3d Bounds = DynMesh.GetBounds();
+		Pivot = Bounds.Center();
+	}
+
+	// Parse transform components
+	FVector3d Translation = FVector3d::Zero();
+	FVector3d Scale = FVector3d::One();
+	FRotator Rotation = FRotator::ZeroRotator;
+	
+	const TSharedPtr<FJsonObject>* LocationObj = nullptr;
+	if (Params->TryGetObjectField(TEXT("location"), LocationObj))
+	{
+		Translation = MeshHelpers::ParseVector(*LocationObj);
+	}
+
+	const TSharedPtr<FJsonObject>* RotationObj = nullptr;
+	if (Params->TryGetObjectField(TEXT("rotation"), RotationObj))
+	{
+		double Pitch = (*RotationObj)->GetNumberField(TEXT("pitch"));
+		double Yaw = (*RotationObj)->GetNumberField(TEXT("yaw"));
+		double Roll = (*RotationObj)->GetNumberField(TEXT("roll"));
+		Rotation = FRotator(Pitch, Yaw, Roll);
+	}
+
+	const TSharedPtr<FJsonObject>* ScaleObj = nullptr;
+	if (Params->TryGetObjectField(TEXT("scale"), ScaleObj))
+	{
+		Scale = MeshHelpers::ParseVector(*ScaleObj, FVector3d::One());
+	}
+
+	// Apply transforms in SRT order (Scale, Rotate, Translate) around pivot
+	// 1. Scale around pivot
+	if (!Scale.Equals(FVector3d::One()))
+	{
+		MeshTransforms::Scale(DynMesh, Scale, Pivot);
+	}
+	
+	// 2. Rotate around pivot
+	if (!Rotation.IsNearlyZero())
+	{
+		MeshTransforms::Rotate(DynMesh, Rotation, Pivot);
+	}
+	
+	// 3. Translate
+	if (Translation.SquaredLength() > UE_SMALL_NUMBER)
+	{
+		MeshTransforms::Translate(DynMesh, Translation);
+	}
+
+	FMeshNormals::QuickComputeVertexNormals(DynMesh);
+
+	UStaticMesh* OutputMesh = nullptr;
+	if (OutputPath == MeshPath)
+	{
+		MeshHelpers::SaveDynamicMeshToStaticMesh(DynMesh, SourceMesh);
+		OutputMesh = SourceMesh;
+	}
+	else
+	{
+		OutputMesh = MeshHelpers::CreateStaticMeshFromDynamicMesh(DynMesh, OutputPath, true, true, &Error);
+		if (!OutputMesh)
+		{
+			return FECACommandResult::Error(Error);
+		}
+	}
+
+	TSharedPtr<FJsonObject> Result = MeshHelpers::BuildMeshResultJson(OutputMesh);
+	return FECACommandResult::Success(Result);
+}
+
+//==============================================================================
+// MESH QUERIES
+//==============================================================================
+
+FECACommandResult FECACommand_GetMeshInfo::Execute(const TSharedPtr<FJsonObject>& Params)
+{
+	FString MeshPath;
+	if (!GetStringParam(Params, TEXT("mesh_path"), MeshPath))
+	{
+		return FECACommandResult::Error(TEXT("Missing required parameter: mesh_path"));
+	}
+
+	UStaticMesh* Mesh = MeshHelpers::LoadStaticMesh(MeshPath);
+	if (!Mesh)
+	{
+		return FECACommandResult::Error(FString::Printf(TEXT("Mesh not found: %s"), *MeshPath));
+	}
+
+	TSharedPtr<FJsonObject> Result = MakeResult();
+	Result->SetStringField(TEXT("mesh_path"), Mesh->GetPathName());
+	Result->SetStringField(TEXT("mesh_name"), Mesh->GetName());
+
+	// Get mesh description stats
+	const FMeshDescription* MeshDesc = Mesh->GetMeshDescription(0);
+	if (MeshDesc)
+	{
+		Result->SetNumberField(TEXT("vertex_count"), MeshDesc->Vertices().Num());
+		Result->SetNumberField(TEXT("triangle_count"), MeshDesc->Triangles().Num());
+		Result->SetNumberField(TEXT("polygon_count"), MeshDesc->Polygons().Num());
+		Result->SetNumberField(TEXT("edge_count"), MeshDesc->Edges().Num());
+	}
+
+	// LOD info
+	Result->SetNumberField(TEXT("lod_count"), Mesh->GetNumLODs());
+
+	// Material slots
+	TArray<TSharedPtr<FJsonValue>> MaterialsArray;
+	for (int32 i = 0; i < Mesh->GetStaticMaterials().Num(); i++)
+	{
+		const FStaticMaterial& Mat = Mesh->GetStaticMaterials()[i];
+		TSharedPtr<FJsonObject> MatObj = MakeShared<FJsonObject>();
+		MatObj->SetNumberField(TEXT("slot_index"), i);
+		MatObj->SetStringField(TEXT("slot_name"), Mat.MaterialSlotName.ToString());
+		if (Mat.MaterialInterface)
+		{
+			MatObj->SetStringField(TEXT("material_path"), Mat.MaterialInterface->GetPathName());
+		}
+		MaterialsArray.Add(MakeShared<FJsonValueObject>(MatObj));
+	}
+	Result->SetArrayField(TEXT("material_slots"), MaterialsArray);
+
+	// Bounds
+	FBoxSphereBounds Bounds = Mesh->GetBounds();
+	TSharedPtr<FJsonObject> BoundsObj = MakeShared<FJsonObject>();
+	BoundsObj->SetNumberField(TEXT("center_x"), Bounds.Origin.X);
+	BoundsObj->SetNumberField(TEXT("center_y"), Bounds.Origin.Y);
+	BoundsObj->SetNumberField(TEXT("center_z"), Bounds.Origin.Z);
+	BoundsObj->SetNumberField(TEXT("extent_x"), Bounds.BoxExtent.X);
+	BoundsObj->SetNumberField(TEXT("extent_y"), Bounds.BoxExtent.Y);
+	BoundsObj->SetNumberField(TEXT("extent_z"), Bounds.BoxExtent.Z);
+	BoundsObj->SetNumberField(TEXT("sphere_radius"), Bounds.SphereRadius);
+	Result->SetObjectField(TEXT("bounds"), BoundsObj);
+
+	// Collision info
+	if (Mesh->GetBodySetup())
+	{
+		Result->SetBoolField(TEXT("has_collision"), true);
+		Result->SetStringField(TEXT("collision_type"), 
+			Mesh->GetBodySetup()->CollisionTraceFlag == CTF_UseComplexAsSimple ? TEXT("complex") : TEXT("simple"));
+	}
+	else
+	{
+		Result->SetBoolField(TEXT("has_collision"), false);
+	}
+
+	return FECACommandResult::Success(Result);
+}
+
+FECACommandResult FECACommand_GetMeshBounds::Execute(const TSharedPtr<FJsonObject>& Params)
+{
+	FString MeshPath;
+	if (!GetStringParam(Params, TEXT("mesh_path"), MeshPath))
+	{
+		return FECACommandResult::Error(TEXT("Missing required parameter: mesh_path"));
+	}
+
+	UStaticMesh* Mesh = MeshHelpers::LoadStaticMesh(MeshPath);
+	if (!Mesh)
+	{
+		return FECACommandResult::Error(FString::Printf(TEXT("Mesh not found: %s"), *MeshPath));
+	}
+
+	FBoxSphereBounds Bounds = Mesh->GetBounds();
+	
+	TSharedPtr<FJsonObject> Result = MakeResult();
+	Result->SetStringField(TEXT("mesh_path"), MeshPath);
+	
+	// Min/Max box
+	TSharedPtr<FJsonObject> MinObj = MakeShared<FJsonObject>();
+	MinObj->SetNumberField(TEXT("x"), Bounds.Origin.X - Bounds.BoxExtent.X);
+	MinObj->SetNumberField(TEXT("y"), Bounds.Origin.Y - Bounds.BoxExtent.Y);
+	MinObj->SetNumberField(TEXT("z"), Bounds.Origin.Z - Bounds.BoxExtent.Z);
+	Result->SetObjectField(TEXT("min"), MinObj);
+
+	TSharedPtr<FJsonObject> MaxObj = MakeShared<FJsonObject>();
+	MaxObj->SetNumberField(TEXT("x"), Bounds.Origin.X + Bounds.BoxExtent.X);
+	MaxObj->SetNumberField(TEXT("y"), Bounds.Origin.Y + Bounds.BoxExtent.Y);
+	MaxObj->SetNumberField(TEXT("z"), Bounds.Origin.Z + Bounds.BoxExtent.Z);
+	Result->SetObjectField(TEXT("max"), MaxObj);
+
+	// Size
+	TSharedPtr<FJsonObject> SizeObj = MakeShared<FJsonObject>();
+	SizeObj->SetNumberField(TEXT("x"), Bounds.BoxExtent.X * 2);
+	SizeObj->SetNumberField(TEXT("y"), Bounds.BoxExtent.Y * 2);
+	SizeObj->SetNumberField(TEXT("z"), Bounds.BoxExtent.Z * 2);
+	Result->SetObjectField(TEXT("size"), SizeObj);
+
+	// Center
+	TSharedPtr<FJsonObject> CenterObj = MakeShared<FJsonObject>();
+	CenterObj->SetNumberField(TEXT("x"), Bounds.Origin.X);
+	CenterObj->SetNumberField(TEXT("y"), Bounds.Origin.Y);
+	CenterObj->SetNumberField(TEXT("z"), Bounds.Origin.Z);
+	Result->SetObjectField(TEXT("center"), CenterObj);
+
+	Result->SetNumberField(TEXT("sphere_radius"), Bounds.SphereRadius);
+
+	return FECACommandResult::Success(Result);
+}
+
+//==============================================================================
+// PROCEDURAL MESH
+//==============================================================================
+
+FECACommandResult FECACommand_SpawnProceduralMesh::Execute(const TSharedPtr<FJsonObject>& Params)
+{
+	FString ActorName;
+	if (!GetStringParam(Params, TEXT("actor_name"), ActorName))
+	{
+		return FECACommandResult::Error(TEXT("Missing required parameter: actor_name"));
+	}
+
+	const TArray<TSharedPtr<FJsonValue>>* VerticesArray;
+	if (!GetArrayParam(Params, TEXT("vertices"), VerticesArray))
+	{
+		return FECACommandResult::Error(TEXT("Missing required parameter: vertices"));
+	}
+
+	const TArray<TSharedPtr<FJsonValue>>* TrianglesArray;
+	if (!GetArrayParam(Params, TEXT("triangles"), TrianglesArray))
+	{
+		return FECACommandResult::Error(TEXT("Missing required parameter: triangles"));
+	}
+
+	UWorld* World = GetEditorWorld();
+	if (!World)
+	{
+		return FECACommandResult::Error(TEXT("No editor world available"));
+	}
+
+	// Parse location/rotation
+	FVector Location = FVector::ZeroVector;
+	FRotator Rotation = FRotator::ZeroRotator;
+
+	const TSharedPtr<FJsonObject>* LocationObj = nullptr;
+	if (Params->TryGetObjectField(TEXT("location"), LocationObj))
+	{
+		Location.X = (*LocationObj)->GetNumberField(TEXT("x"));
+		Location.Y = (*LocationObj)->GetNumberField(TEXT("y"));
+		Location.Z = (*LocationObj)->GetNumberField(TEXT("z"));
+	}
+
+	const TSharedPtr<FJsonObject>* RotationObj = nullptr;
+	if (Params->TryGetObjectField(TEXT("rotation"), RotationObj))
+	{
+		Rotation.Pitch = (*RotationObj)->GetNumberField(TEXT("pitch"));
+		Rotation.Yaw = (*RotationObj)->GetNumberField(TEXT("yaw"));
+		Rotation.Roll = (*RotationObj)->GetNumberField(TEXT("roll"));
+	}
+
+	// Spawn actor
+	FActorSpawnParameters SpawnParams;
+	SpawnParams.Name = FName(*ActorName);
+	AActor* NewActor = World->SpawnActor<AActor>(AActor::StaticClass(), Location, Rotation, SpawnParams);
+	if (!NewActor)
+	{
+		return FECACommandResult::Error(TEXT("Failed to spawn actor"));
+	}
+
+	NewActor->SetActorLabel(ActorName);
+
+	// Add procedural mesh component
+	UProceduralMeshComponent* ProcMesh = NewObject<UProceduralMeshComponent>(NewActor, TEXT("ProceduralMesh"));
+	ProcMesh->RegisterComponent();
+	NewActor->SetRootComponent(ProcMesh);
+
+	// Parse vertices
+	TArray<FVector> Vertices;
+	for (const TSharedPtr<FJsonValue>& VertexValue : *VerticesArray)
+	{
+		const TSharedPtr<FJsonObject>* VertexObj;
+		if (VertexValue->TryGetObject(VertexObj))
+		{
+			Vertices.Add(FVector(
+				(*VertexObj)->GetNumberField(TEXT("x")),
+				(*VertexObj)->GetNumberField(TEXT("y")),
+				(*VertexObj)->GetNumberField(TEXT("z"))
+			));
+		}
+	}
+
+	// Parse triangles (flat array)
+	TArray<int32> Triangles;
+	for (const TSharedPtr<FJsonValue>& IndexValue : *TrianglesArray)
+	{
+		Triangles.Add((int32)IndexValue->AsNumber());
+	}
+
+	// Parse optional arrays
+	TArray<FVector> Normals;
+	TArray<FVector2D> UVs;
+	TArray<FLinearColor> VertexColors;
+	TArray<FProcMeshTangent> Tangents;
+
+	// Create mesh section
+	ProcMesh->CreateMeshSection_LinearColor(0, Vertices, Triangles, Normals, UVs, VertexColors, Tangents, true);
+
+	// Apply material if specified
+	FString MaterialPath;
+	if (GetStringParam(Params, TEXT("material_path"), MaterialPath, false))
+	{
+		UMaterialInterface* Material = LoadObject<UMaterialInterface>(nullptr, *MaterialPath);
+		if (Material)
+		{
+			ProcMesh->SetMaterial(0, Material);
+		}
+	}
+
+	// Enable collision if specified
+	bool bEnableCollision = true;
+	GetBoolParam(Params, TEXT("enable_collision"), bEnableCollision, false);
+	ProcMesh->SetCollisionEnabled(bEnableCollision ? ECollisionEnabled::QueryAndPhysics : ECollisionEnabled::NoCollision);
+
+	TSharedPtr<FJsonObject> Result = MakeResult();
+	Result->SetStringField(TEXT("actor_name"), NewActor->GetActorLabel());
+	Result->SetNumberField(TEXT("vertex_count"), Vertices.Num());
+	Result->SetNumberField(TEXT("triangle_count"), Triangles.Num() / 3);
+	
+	TSharedPtr<FJsonObject> LocationResult = MakeShared<FJsonObject>();
+	LocationResult->SetNumberField(TEXT("x"), Location.X);
+	LocationResult->SetNumberField(TEXT("y"), Location.Y);
+	LocationResult->SetNumberField(TEXT("z"), Location.Z);
+	Result->SetObjectField(TEXT("location"), LocationResult);
+
+	return FECACommandResult::Success(Result);
+}
+
+FECACommandResult FECACommand_UpdateProceduralMesh::Execute(const TSharedPtr<FJsonObject>& Params)
+{
+	FString ActorName;
+	if (!GetStringParam(Params, TEXT("actor_name"), ActorName))
+	{
+		return FECACommandResult::Error(TEXT("Missing required parameter: actor_name"));
+	}
+
+	// Find the actor
+	AActor* FoundActor = FindActorByName(ActorName);
+	if (!FoundActor)
+	{
+		return FECACommandResult::Error(FString::Printf(TEXT("Actor not found: %s"), *ActorName));
+	}
+
+	// Find procedural mesh component
+	UProceduralMeshComponent* ProcMesh = FoundActor->FindComponentByClass<UProceduralMeshComponent>();
+	if (!ProcMesh)
+	{
+		return FECACommandResult::Error(TEXT("Actor does not have a ProceduralMeshComponent"));
+	}
+
+	const TArray<TSharedPtr<FJsonValue>>* VerticesArray;
+	if (!GetArrayParam(Params, TEXT("vertices"), VerticesArray))
+	{
+		return FECACommandResult::Error(TEXT("Missing required parameter: vertices"));
+	}
+
+	const TArray<TSharedPtr<FJsonValue>>* TrianglesArray;
+	if (!GetArrayParam(Params, TEXT("triangles"), TrianglesArray))
+	{
+		return FECACommandResult::Error(TEXT("Missing required parameter: triangles"));
+	}
+
+	int32 SectionIndex = 0;
+	GetIntParam(Params, TEXT("section_index"), SectionIndex, false);
+
+	// Parse data
+	TArray<FVector> Vertices;
+	for (const TSharedPtr<FJsonValue>& VertexValue : *VerticesArray)
+	{
+		const TSharedPtr<FJsonObject>* VertexObj;
+		if (VertexValue->TryGetObject(VertexObj))
+		{
+			Vertices.Add(FVector(
+				(*VertexObj)->GetNumberField(TEXT("x")),
+				(*VertexObj)->GetNumberField(TEXT("y")),
+				(*VertexObj)->GetNumberField(TEXT("z"))
+			));
+		}
+	}
+
+	TArray<int32> Triangles;
+	for (const TSharedPtr<FJsonValue>& IndexValue : *TrianglesArray)
+	{
+		Triangles.Add((int32)IndexValue->AsNumber());
+	}
+
+	// Update mesh section
+	TArray<FVector> Normals;
+	TArray<FVector2D> UVs;
+	TArray<FLinearColor> VertexColors;
+	TArray<FProcMeshTangent> Tangents;
+
+	ProcMesh->CreateMeshSection_LinearColor(SectionIndex, Vertices, Triangles, Normals, UVs, VertexColors, Tangents, true);
+
+	TSharedPtr<FJsonObject> Result = MakeResult();
+	Result->SetStringField(TEXT("actor_name"), ActorName);
+	Result->SetNumberField(TEXT("section_index"), SectionIndex);
+	Result->SetNumberField(TEXT("vertex_count"), Vertices.Num());
+	Result->SetNumberField(TEXT("triangle_count"), Triangles.Num() / 3);
+
+	return FECACommandResult::Success(Result);
+}
+
+//==============================================================================
+// MESH COMBINE
+//==============================================================================
+
+FECACommandResult FECACommand_MeshCombine::Execute(const TSharedPtr<FJsonObject>& Params)
+{
+	FString OutputPath;
+	if (!GetStringParam(Params, TEXT("output_path"), OutputPath))
+	{
+		return FECACommandResult::Error(TEXT("Missing required parameter: output_path"));
+	}
+
+	const TArray<TSharedPtr<FJsonValue>>* MeshPathsArray;
+	if (!GetArrayParam(Params, TEXT("mesh_paths"), MeshPathsArray))
+	{
+		return FECACommandResult::Error(TEXT("Missing required parameter: mesh_paths"));
+	}
+
+	if (MeshPathsArray->Num() < 1)
+	{
+		return FECACommandResult::Error(TEXT("mesh_paths must contain at least one mesh"));
+	}
+
+	// Optional transforms array
+	const TArray<TSharedPtr<FJsonValue>>* TransformsArray = nullptr;
+	GetArrayParam(Params, TEXT("transforms"), TransformsArray, false);
+
+	// Combine meshes
+	FDynamicMesh3 CombinedMesh;
+	CombinedMesh.EnableAttributes();
+	CombinedMesh.Attributes()->SetNumUVLayers(1);
+
+	FString Error;
+	for (int32 i = 0; i < MeshPathsArray->Num(); i++)
+	{
+		FString MeshPath = (*MeshPathsArray)[i]->AsString();
+		UStaticMesh* Mesh = MeshHelpers::LoadStaticMesh(MeshPath);
+		if (!Mesh)
+		{
+			return FECACommandResult::Error(FString::Printf(TEXT("Mesh not found: %s"), *MeshPath));
+		}
+
+		FDynamicMesh3 TempMesh;
+		if (!MeshHelpers::LoadStaticMeshIntoDynamicMesh(Mesh, TempMesh, &Error))
+		{
+			return FECACommandResult::Error(FString::Printf(TEXT("Failed to load mesh %s: %s"), *MeshPath, *Error));
+		}
+
+		// Apply transform if provided
+		if (TransformsArray && i < TransformsArray->Num())
+		{
+			const TSharedPtr<FJsonObject>* TransformObj;
+			if ((*TransformsArray)[i]->TryGetObject(TransformObj))
+			{
+				FTransform Transform = MeshHelpers::ParseTransform(*TransformObj);
+				// Convert FTransform to FTransformSRT3d for mesh transformation
+				FTransform3d Transform3d_UE(Transform);
+				FTransformSRT3d Transform3d(Transform3d_UE);
+				MeshTransforms::ApplyTransform(TempMesh, Transform3d, true);
+			}
+		}
+
+		// Append to combined mesh
+		FMeshIndexMappings Mappings;
+		CombinedMesh.EnableMatchingAttributes(TempMesh);
+		FDynamicMeshEditor Editor(&CombinedMesh);
+		Editor.AppendMesh(&TempMesh, Mappings);
+		
+		// AppendMesh doesn't transfer UV/normal overlays - we need to do it explicitly
+		if (TempMesh.HasAttributes() && CombinedMesh.HasAttributes())
+		{
+			// Append UV overlays
+			for (int32 UVLayerIdx = 0; UVLayerIdx < TempMesh.Attributes()->NumUVLayers(); UVLayerIdx++)
+			{
+				if (UVLayerIdx < CombinedMesh.Attributes()->NumUVLayers())
+				{
+					const FDynamicMeshUVOverlay* FromUVs = TempMesh.Attributes()->GetUVLayer(UVLayerIdx);
+					FDynamicMeshUVOverlay* ToUVs = CombinedMesh.Attributes()->GetUVLayer(UVLayerIdx);
+					if (FromUVs && ToUVs && FromUVs->ElementCount() > 0)
+					{
+						FIndexMapi UVMap;
+						Editor.AppendUVs(&TempMesh, FromUVs, ToUVs, Mappings.GetVertexMap(), Mappings.GetTriangleMap(), UVMap);
+					}
+				}
+			}
+			
+			// Append normal overlay
+			const FDynamicMeshNormalOverlay* FromNormals = TempMesh.Attributes()->PrimaryNormals();
+			FDynamicMeshNormalOverlay* ToNormals = CombinedMesh.Attributes()->PrimaryNormals();
+			if (FromNormals && ToNormals && FromNormals->ElementCount() > 0)
+			{
+				FIndexMapi NormalMap;
+				Editor.AppendNormals(&TempMesh, FromNormals, ToNormals, Mappings.GetVertexMap(), Mappings.GetTriangleMap(), nullptr, NormalMap);
+			}
+		}
+	}
+
+	FMeshNormals::QuickComputeVertexNormals(CombinedMesh);
+
+	UStaticMesh* OutputMesh = MeshHelpers::CreateStaticMeshFromDynamicMesh(CombinedMesh, OutputPath, true, true, &Error);
+	if (!OutputMesh)
+	{
+		return FECACommandResult::Error(Error);
+	}
+
+	TSharedPtr<FJsonObject> Result = MeshHelpers::BuildMeshResultJson(OutputMesh, CombinedMesh.VertexCount(), CombinedMesh.TriangleCount());
+	Result->SetNumberField(TEXT("combined_count"), MeshPathsArray->Num());
+	return FECACommandResult::Success(Result);
+}
+
+//==============================================================================
+// IMPORT / EXPORT
+//==============================================================================
+
+FECACommandResult FECACommand_ImportMesh::Execute(const TSharedPtr<FJsonObject>& Params)
+{
+	return FECACommandResult::Error(TEXT("import_mesh: Use 'import_obj' command for OBJ files, or use UE's content browser for FBX/glTF"));
+}
+
+FECACommandResult FECACommand_ExportMesh::Execute(const TSharedPtr<FJsonObject>& Params)
+{
+	FString MeshPath;
+	if (!GetStringParam(Params, TEXT("mesh_path"), MeshPath))
+	{
+		return FECACommandResult::Error(TEXT("Missing required parameter: mesh_path"));
+	}
+
+	FString OutputFile;
+	if (!GetStringParam(Params, TEXT("output_file"), OutputFile))
+	{
+		return FECACommandResult::Error(TEXT("Missing required parameter: output_file"));
+	}
+
+	UStaticMesh* Mesh = MeshHelpers::LoadStaticMesh(MeshPath);
+	if (!Mesh)
+	{
+		return FECACommandResult::Error(FString::Printf(TEXT("Mesh not found: %s"), *MeshPath));
+	}
+
+	FDynamicMesh3 DynMesh;
+	FString Error;
+	if (!MeshHelpers::LoadStaticMeshIntoDynamicMesh(Mesh, DynMesh, &Error))
+	{
+		return FECACommandResult::Error(Error);
+	}
+
+	// Export to OBJ format
+	FString Extension = FPaths::GetExtension(OutputFile).ToLower();
+	if (Extension == TEXT("obj"))
+	{
+		FString ObjContent;
+		ObjContent += TEXT("# Exported from ECABridge\n");
+		ObjContent += FString::Printf(TEXT("# Vertices: %d, Triangles: %d\n\n"), DynMesh.VertexCount(), DynMesh.TriangleCount());
+
+		// Write vertices
+		TMap<int32, int32> VertexIDToIndex;
+		int32 Index = 1;
+		for (int32 VertID : DynMesh.VertexIndicesItr())
+		{
+			FVector3d Pos = DynMesh.GetVertex(VertID);
+			ObjContent += FString::Printf(TEXT("v %f %f %f\n"), Pos.X, Pos.Z, Pos.Y); // Convert coordinate system
+			VertexIDToIndex.Add(VertID, Index++);
+		}
+
+		ObjContent += TEXT("\n");
+
+		// Write faces (1-indexed)
+		for (int32 TriID : DynMesh.TriangleIndicesItr())
+		{
+			FIndex3i Tri = DynMesh.GetTriangle(TriID);
+			ObjContent += FString::Printf(TEXT("f %d %d %d\n"), 
+				VertexIDToIndex[Tri.A], 
+				VertexIDToIndex[Tri.B], 
+				VertexIDToIndex[Tri.C]);
+		}
+
+		if (!FFileHelper::SaveStringToFile(ObjContent, *OutputFile))
+		{
+			return FECACommandResult::Error(FString::Printf(TEXT("Failed to write file: %s"), *OutputFile));
+		}
+	}
+	else
+	{
+		return FECACommandResult::Error(TEXT("Only OBJ export is currently supported"));
+	}
+
+	TSharedPtr<FJsonObject> Result = MakeResult();
+	Result->SetStringField(TEXT("output_file"), OutputFile);
+	Result->SetNumberField(TEXT("vertex_count"), DynMesh.VertexCount());
+	Result->SetNumberField(TEXT("triangle_count"), DynMesh.TriangleCount());
+
+	return FECACommandResult::Success(Result);
+}
+
+FECACommandResult FECACommand_MeshVoronoiFracture::Execute(const TSharedPtr<FJsonObject>& Params)
+{
+	FString MeshPath;
+	if (!GetStringParam(Params, TEXT("mesh_path"), MeshPath))
+	{
+		return FECACommandResult::Error(TEXT("Missing required parameter: mesh_path"));
+	}
+
+	FString OutputBasePath;
+	if (!GetStringParam(Params, TEXT("output_base_path"), OutputBasePath))
+	{
+		return FECACommandResult::Error(TEXT("Missing required parameter: output_base_path"));
+	}
+
+	int32 NumPieces = 5;
+	GetIntParam(Params, TEXT("num_pieces"), NumPieces, false);
+	NumPieces = FMath::Clamp(NumPieces, 2, 50);
+
+	// Optional seed points - if not provided, generate random points
+	const TArray<TSharedPtr<FJsonValue>>* SeedPointsArray = nullptr;
+	GetArrayParam(Params, TEXT("seed_points"), SeedPointsArray, false);
+
+	UStaticMesh* SourceMesh = MeshHelpers::LoadStaticMesh(MeshPath);
+	if (!SourceMesh)
+	{
+		return FECACommandResult::Error(FString::Printf(TEXT("Mesh not found: %s"), *MeshPath));
+	}
+
+	FDynamicMesh3 SourceDynMesh;
+	FString Error;
+	if (!MeshHelpers::LoadStaticMeshIntoDynamicMesh(SourceMesh, SourceDynMesh, &Error))
+	{
+		return FECACommandResult::Error(Error);
+	}
+
+	// Get mesh bounds
+	FAxisAlignedBox3d Bounds = SourceDynMesh.GetBounds();
+	FVector3d Center = Bounds.Center();
+	FVector3d Size = Bounds.Diagonal();
+
+	// Generate or parse seed points for Voronoi cells
+	TArray<FVector3d> SeedPoints;
+	if (SeedPointsArray && SeedPointsArray->Num() > 0)
+	{
+		for (const TSharedPtr<FJsonValue>& PointValue : *SeedPointsArray)
+		{
+			const TSharedPtr<FJsonObject>* PointObj;
+			if (PointValue->TryGetObject(PointObj))
+			{
+				SeedPoints.Add(FVector3d(
+					(*PointObj)->GetNumberField(TEXT("x")),
+					(*PointObj)->GetNumberField(TEXT("y")),
+					(*PointObj)->GetNumberField(TEXT("z"))
+				));
+			}
+		}
+	}
+	else
+	{
+		// Generate random seed points within the mesh bounds
+		FRandomStream Random(FMath::Rand());
+		for (int32 i = 0; i < NumPieces; i++)
+		{
+			FVector3d Point(
+				Bounds.Min.X + Random.FRand() * Size.X,
+				Bounds.Min.Y + Random.FRand() * Size.Y,
+				Bounds.Min.Z + Random.FRand() * Size.Z
+			);
+			SeedPoints.Add(Point);
+		}
+	}
+
+	if (SeedPoints.Num() < 2)
+	{
+		return FECACommandResult::Error(TEXT("Need at least 2 seed points for fracture"));
+	}
+
+	// Simple Voronoi-like fracture using plane cuts
+	// For each seed point, cut the mesh with planes perpendicular to vectors to other seeds
+	TArray<FString> OutputPaths;
+	TArray<TSharedPtr<FJsonObject>> PieceResults;
+
+	for (int32 i = 0; i < SeedPoints.Num(); i++)
+	{
+		FDynamicMesh3 PieceMesh(SourceDynMesh);
+		FVector3d CurrentSeed = SeedPoints[i];
+		
+		// Cut against planes defined by midpoints to other seeds
+		for (int32 j = 0; j < SeedPoints.Num(); j++)
+		{
+			if (i == j) continue;
+			
+			FVector3d OtherSeed = SeedPoints[j];
+			FVector3d MidPoint = (CurrentSeed + OtherSeed) * 0.5;
+			FVector3d PlaneNormal = (CurrentSeed - OtherSeed).GetSafeNormal();
+			
+			// Cut and keep the side containing our seed point
+			FMeshPlaneCut PlaneCut(&PieceMesh, MidPoint, PlaneNormal);
+			PlaneCut.Cut();
+			
+			// Fill holes from the cut
+			for (const FMeshPlaneCut::FOpenBoundary& Boundary : PlaneCut.OpenBoundaries)
+			{
+				for (const FEdgeLoop& Loop : Boundary.CutLoops)
+				{
+					FSimpleHoleFiller Filler(&PieceMesh, Loop);
+					Filler.Fill();
+				}
+			}
+		}
+		
+		// Only save if the piece has geometry
+		if (PieceMesh.TriangleCount() > 0)
+		{
+			FMeshNormals::QuickComputeVertexNormals(PieceMesh);
+			
+			FString PiecePath = FString::Printf(TEXT("%s_piece_%d"), *OutputBasePath, i);
+			UStaticMesh* PieceMeshAsset = MeshHelpers::CreateStaticMeshFromDynamicMesh(PieceMesh, PiecePath, true, true, &Error);
+			
+			if (PieceMeshAsset)
+			{
+				OutputPaths.Add(PieceMeshAsset->GetPathName());
+				
+				TSharedPtr<FJsonObject> PieceInfo = MakeShared<FJsonObject>();
+				PieceInfo->SetNumberField(TEXT("piece_index"), i);
+				PieceInfo->SetStringField(TEXT("mesh_path"), PieceMeshAsset->GetPathName());
+				PieceInfo->SetNumberField(TEXT("triangle_count"), PieceMesh.TriangleCount());
+				
+				TSharedPtr<FJsonObject> SeedObj = MakeShared<FJsonObject>();
+				SeedObj->SetNumberField(TEXT("x"), CurrentSeed.X);
+				SeedObj->SetNumberField(TEXT("y"), CurrentSeed.Y);
+				SeedObj->SetNumberField(TEXT("z"), CurrentSeed.Z);
+				PieceInfo->SetObjectField(TEXT("seed_point"), SeedObj);
+				
+				PieceResults.Add(PieceInfo);
+			}
+		}
+	}
+
+	if (PieceResults.Num() == 0)
+	{
+		return FECACommandResult::Error(TEXT("Fracture produced no valid pieces"));
+	}
+
+	TSharedPtr<FJsonObject> Result = MakeResult();
+	Result->SetNumberField(TEXT("num_pieces"), PieceResults.Num());
+	Result->SetNumberField(TEXT("requested_pieces"), NumPieces);
+	
+	TArray<TSharedPtr<FJsonValue>> PiecesArray;
+	for (const TSharedPtr<FJsonObject>& PieceInfo : PieceResults)
+	{
+		PiecesArray.Add(MakeShared<FJsonValueObject>(PieceInfo));
+	}
+	Result->SetArrayField(TEXT("pieces"), PiecesArray);
+	
+	return FECACommandResult::Success(Result);
+}
+
+//==============================================================================
+// MESH DISPLACE FROM DEPTH
+//==============================================================================
+
+FECACommandResult FECACommand_MeshDisplaceFromDepth::Execute(const TSharedPtr<FJsonObject>& Params)
+{
+	// Get mesh path
+	FString MeshPath;
+	if (!GetStringParam(Params, TEXT("mesh_path"), MeshPath))
+	{
+		return FECACommandResult::Error(TEXT("Missing required parameter: mesh_path"));
+	}
+
+	// Get actor path to get world transform
+	FString ActorPath;
+	if (!GetStringParam(Params, TEXT("actor_path"), ActorPath))
+	{
+		return FECACommandResult::Error(TEXT("Missing required parameter: actor_path. Provide the path to the actor instance in the world."));
+	}
+
+	// Find the actor and get its world transform
+	AActor* MeshActor = FindObject<AActor>(nullptr, *ActorPath);
+	if (!MeshActor)
+	{
+		// Try loading it
+		MeshActor = LoadObject<AActor>(nullptr, *ActorPath);
+	}
+	if (!MeshActor)
+	{
+		return FECACommandResult::Error(FString::Printf(TEXT("Actor not found: %s"), *ActorPath));
+	}
+
+	FTransform ActorWorldTransform = MeshActor->GetActorTransform();
+
+	// Get output path (optional, defaults to overwriting input)
+	FString OutputPath = MeshPath;
+	GetStringParam(Params, TEXT("output_path"), OutputPath, false);
+
+	// Get displacement parameters
+	double DisplacementStrength = 100.0;
+	GetFloatParam(Params, TEXT("strength"), DisplacementStrength, false);
+
+	bool bUseAbsolute = true;  // Default to simpler absolute mode
+	GetBoolParam(Params, TEXT("use_absolute"), bUseAbsolute, false);
+
+	int32 CenterValue = 128;  // Middle gray = no displacement
+	GetIntParam(Params, TEXT("center_value"), CenterValue, false);
+	float NormalizedCenter = CenterValue / 255.0f;
+
+	double MaxDepth = 10000.0;
+	GetFloatParam(Params, TEXT("max_depth"), MaxDepth, false);
+
+	bool bInvert = false;
+	GetBoolParam(Params, TEXT("invert"), bInvert, false);
+
+	int32 SmoothIterations = 3;
+	GetIntParam(Params, TEXT("smooth_iterations"), SmoothIterations, false);
+
+	// Target depth image path
+	FString TargetDepthImagePath;
+	bool bUseTargetImage = GetStringParam(Params, TEXT("target_depth_image"), TargetDepthImagePath, false);
+
+	// Load the mesh
+	UStaticMesh* SourceMesh = MeshHelpers::LoadStaticMesh(MeshPath);
+	if (!SourceMesh)
+	{
+		return FECACommandResult::Error(FString::Printf(TEXT("Mesh not found: %s"), *MeshPath));
+	}
+
+	// Convert to dynamic mesh
+	FDynamicMesh3 DynMesh;
+	FString Error;
+	if (!MeshHelpers::LoadStaticMeshIntoDynamicMesh(SourceMesh, DynMesh, &Error))
+	{
+		return FECACommandResult::Error(Error);
+	}
+
+	// Get the editor viewport
+	FLevelEditorModule& LevelEditorModule = FModuleManager::GetModuleChecked<FLevelEditorModule>(TEXT("LevelEditor"));
+	TSharedPtr<IAssetViewport> ActiveViewport = LevelEditorModule.GetFirstActiveViewport();
+
+	if (!ActiveViewport.IsValid())
+	{
+		return FECACommandResult::Error(TEXT("No active viewport"));
+	}
+
+	FEditorViewportClient* ViewportClient = static_cast<FEditorViewportClient*>(ActiveViewport->GetActiveViewport()->GetClient());
+	if (!ViewportClient)
+	{
+		return FECACommandResult::Error(TEXT("Could not get viewport client"));
+	}
+
+	UWorld* World = GEditor->GetEditorWorldContext().World();
+	if (!World)
+	{
+		return FECACommandResult::Error(TEXT("No world available"));
+	}
+
+	FVector CameraLocation = ViewportClient->GetViewLocation();
+	FRotator CameraRotation = ViewportClient->GetViewRotation();
+	float FOV = ViewportClient->ViewFOV;
+
+	int32 Resolution = 512;
+	GetIntParam(Params, TEXT("resolution"), Resolution, false);
+
+	// Capture current depth
+	UTextureRenderTarget2D* RenderTarget = NewObject<UTextureRenderTarget2D>();
+	RenderTarget->InitCustomFormat(Resolution, Resolution, PF_FloatRGBA, false);
+	RenderTarget->UpdateResourceImmediate(true);
+
+	FActorSpawnParameters SpawnParams;
+	SpawnParams.SpawnCollisionHandlingOverride = ESpawnActorCollisionHandlingMethod::AlwaysSpawn;
+	ASceneCapture2D* CaptureActor = World->SpawnActor<ASceneCapture2D>(CameraLocation, CameraRotation, SpawnParams);
+
+	if (!CaptureActor)
+	{
+		return FECACommandResult::Error(TEXT("Failed to create scene capture actor"));
+	}
+
+	USceneCaptureComponent2D* CaptureComponent = CaptureActor->GetCaptureComponent2D();
+	CaptureComponent->TextureTarget = RenderTarget;
+	CaptureComponent->CaptureSource = ESceneCaptureSource::SCS_SceneDepth;
+	CaptureComponent->FOVAngle = FOV;
+	CaptureComponent->bCaptureEveryFrame = false;
+	CaptureComponent->bCaptureOnMovement = false;
+	CaptureComponent->MaxViewDistanceOverride = MaxDepth;
+
+	CaptureComponent->CaptureScene();
+
+	TArray<FFloat16Color> FloatPixels;
+	FTextureRenderTargetResource* RTResource = RenderTarget->GameThread_GetRenderTargetResource();
+	if (RTResource)
+	{
+		RTResource->ReadFloat16Pixels(FloatPixels);
+	}
+
+	CaptureActor->Destroy();
+
+	if (FloatPixels.Num() == 0)
+	{
+		return FECACommandResult::Error(TEXT("Failed to capture depth buffer"));
+	}
+
+	// Convert captured depth to normalized 0-1 range
+	// SCS_SceneDepth returns world-space distance from camera
+	TArray<float> CurrentDepthBuffer;
+	CurrentDepthBuffer.SetNum(FloatPixels.Num());
+	
+	// First pass: find min/max depth for normalization
+	float CapturedMinDepth = FLT_MAX;
+	float CapturedMaxDepth = 0.0f;
+	for (int32 i = 0; i < FloatPixels.Num(); i++)
+	{
+		float Depth = FloatPixels[i].R.GetFloat();
+		if (Depth > 0.0f && Depth < MaxDepth)
+		{
+			CapturedMinDepth = FMath::Min(CapturedMinDepth, Depth);
+			CapturedMaxDepth = FMath::Max(CapturedMaxDepth, Depth);
+		}
+	}
+	
+	// Normalize current depth to 0-1 range (0 = near, 1 = far)
+	float DepthRange = CapturedMaxDepth - CapturedMinDepth;
+	if (DepthRange < 1.0f) DepthRange = MaxDepth;  // Fallback
+	
+	for (int32 i = 0; i < FloatPixels.Num(); i++)
+	{
+		float Depth = FloatPixels[i].R.GetFloat();
+		// Clamp and normalize to 0-1
+		Depth = FMath::Clamp(Depth, CapturedMinDepth, CapturedMaxDepth);
+		CurrentDepthBuffer[i] = (Depth - CapturedMinDepth) / DepthRange;
+	}
+
+	// Load or generate target depth buffer
+	TArray<float> TargetDepthBuffer;
+	
+	if (bUseTargetImage)
+	{
+		TArray<uint8> FileData;
+		if (!FFileHelper::LoadFileToArray(FileData, *TargetDepthImagePath))
+		{
+			return FECACommandResult::Error(FString::Printf(TEXT("Failed to load target depth image: %s"), *TargetDepthImagePath));
+		}
+
+		IImageWrapperModule& ImageWrapperModule = FModuleManager::LoadModuleChecked<IImageWrapperModule>(FName("ImageWrapper"));
+		TSharedPtr<IImageWrapper> ImageWrapper = ImageWrapperModule.CreateImageWrapper(EImageFormat::PNG);
+
+		if (!ImageWrapper->SetCompressed(FileData.GetData(), FileData.Num()))
+		{
+			return FECACommandResult::Error(TEXT("Failed to decode target depth image"));
+		}
+
+		TArray<uint8> RawData;
+		if (!ImageWrapper->GetRaw(ERGBFormat::Gray, 8, RawData))
+		{
+			return FECACommandResult::Error(TEXT("Failed to get raw pixels from target depth image"));
+		}
+
+		// Target image needs to be resized to match capture resolution
+		int32 TargetWidth = ImageWrapper->GetWidth();
+		int32 TargetHeight = ImageWrapper->GetHeight();
+		
+		// Create normalized 0-1 buffer from grayscale (0=black=near, 255=white=far)
+		TArray<float> TargetRaw;
+		TargetRaw.SetNum(TargetWidth * TargetHeight);
+		for (int32 i = 0; i < RawData.Num() && i < TargetRaw.Num(); i++)
+		{
+			// Normalize 0-255 to 0-1
+			TargetRaw[i] = static_cast<float>(RawData[i]) / 255.0f;
+		}
+		
+		// Resize to match capture resolution if needed
+		if (TargetWidth != Resolution || TargetHeight != Resolution)
+		{
+			TargetDepthBuffer.SetNum(Resolution * Resolution);
+			for (int32 y = 0; y < Resolution; y++)
+			{
+				for (int32 x = 0; x < Resolution; x++)
+				{
+					// Bilinear sample from source
+					float SrcX = (float)x / (Resolution - 1) * (TargetWidth - 1);
+					float SrcY = (float)y / (Resolution - 1) * (TargetHeight - 1);
+					
+					int32 X0 = FMath::Clamp((int32)SrcX, 0, TargetWidth - 1);
+					int32 Y0 = FMath::Clamp((int32)SrcY, 0, TargetHeight - 1);
+					int32 X1 = FMath::Min(X0 + 1, TargetWidth - 1);
+					int32 Y1 = FMath::Min(Y0 + 1, TargetHeight - 1);
+					
+					float FracX = SrcX - X0;
+					float FracY = SrcY - Y0;
+					
+					float V00 = TargetRaw[Y0 * TargetWidth + X0];
+					float V10 = TargetRaw[Y0 * TargetWidth + X1];
+					float V01 = TargetRaw[Y1 * TargetWidth + X0];
+					float V11 = TargetRaw[Y1 * TargetWidth + X1];
+					
+					float Top = V00 + (V10 - V00) * FracX;
+					float Bottom = V01 + (V11 - V01) * FracX;
+					TargetDepthBuffer[y * Resolution + x] = Top + (Bottom - Top) * FracY;
+				}
+			}
+		}
+		else
+		{
+			TargetDepthBuffer = MoveTemp(TargetRaw);
+		}
+	}
+	else
+	{
+		TargetDepthBuffer.SetNumZeroed(CurrentDepthBuffer.Num());
+	}
+
+	// Compute vertex normals for facing checks
+	FMeshNormals MeshNormals(&DynMesh);
+	MeshNormals.ComputeVertexNormals();
+
+	// Camera vectors - simple and direct
+	FVector CamForward = CameraRotation.Vector();
+	// Instead of complex matrix math, project vertices directly using simple geometry
+	// This avoids all the matrix convention confusion
+	const float HalfFOVRadians = FMath::DegreesToRadians(FOV * 0.5f);
+	const float TanHalfFOV = FMath::Tan(HalfFOVRadians);
+	
+	// Camera basis vectors
+	FVector CamRight = FRotationMatrix(CameraRotation).GetScaledAxis(EAxis::Y);
+	FVector CamUp = FRotationMatrix(CameraRotation).GetScaledAxis(EAxis::Z);
+	
+	// Lambda to project a world position to UV coordinates (0-1 range)
+	auto ProjectToUV = [&](const FVector3d& WorldPos) -> FVector2D
+	{
+		// Vector from camera to vertex
+		FVector ToVertex = FVector(WorldPos) - CameraLocation;
+		
+		// Distance along camera forward axis
+		float ForwardDist = FVector::DotProduct(ToVertex, CamForward);
+		
+		// If behind camera, return invalid UV
+		if (ForwardDist <= 0.001f)
+		{
+			return FVector2D(-1, -1);
+		}
+		
+		// Project onto camera plane (perpendicular to forward)
+		float RightDist = FVector::DotProduct(ToVertex, CamRight);
+		float UpDist = FVector::DotProduct(ToVertex, CamUp);
+		
+		// Convert to NDC using perspective projection
+		// At distance ForwardDist, the visible half-width is ForwardDist * tan(FOV/2)
+		float HalfWidth = ForwardDist * TanHalfFOV;
+		
+		float NdcX = RightDist / HalfWidth;  // -1 to 1
+		float NdcY = UpDist / HalfWidth;     // -1 to 1
+		
+		// Check if outside frustum
+		if (FMath::Abs(NdcX) > 1.0f || FMath::Abs(NdcY) > 1.0f)
+		{
+			return FVector2D(-2, -2);
+		}
+		
+		// NDC to UV: X: -1..1 -> 0..1, Y: -1..1 -> 1..0 (flip Y for image coords)
+		float U = (NdcX + 1.0f) * 0.5f;
+		float V = (1.0f - NdcY) * 0.5f;
+		
+		return FVector2D(U, V);
+	};
+
+	// Bilinear sample helper lambda
+	auto SampleBilinear = [&](const TArray<float>& Buffer, float U, float V) -> float
+	{
+		if (Buffer.Num() == 0) return 0.0f;
+		
+		U = FMath::Clamp(U, 0.0f, 1.0f);
+		V = FMath::Clamp(V, 0.0f, 1.0f);
+		
+		float PixelXf = U * (Resolution - 1);
+		float PixelYf = V * (Resolution - 1);
+
+		int32 X0 = FMath::Clamp((int32)FMath::Floor(PixelXf), 0, Resolution - 1);
+		int32 Y0 = FMath::Clamp((int32)FMath::Floor(PixelYf), 0, Resolution - 1);
+		int32 X1 = FMath::Min(X0 + 1, Resolution - 1);
+		int32 Y1 = FMath::Min(Y0 + 1, Resolution - 1);
+
+		float FracX = PixelXf - X0;
+		float FracY = PixelYf - Y0;
+
+		float V00 = Buffer[Y0 * Resolution + X0];
+		float V10 = Buffer[Y0 * Resolution + X1];
+		float V01 = Buffer[Y1 * Resolution + X0];
+		float V11 = Buffer[Y1 * Resolution + X1];
+
+		float Top = V00 + (V10 - V00) * FracX;
+		float Bottom = V01 + (V11 - V01) * FracX;
+		return Top + (Bottom - Top) * FracY;
+	};
+	
+	// Check if a UV location is at a depth discontinuity (edge) in the buffer
+	// Returns true if the local neighborhood has high depth variance
+	auto IsDepthEdge = [&](const TArray<float>& Buffer, float U, float V, float Threshold) -> bool
+	{
+		if (Buffer.Num() == 0) return false;
+		
+		int32 X = FMath::Clamp((int32)(U * (Resolution - 1)), 0, Resolution - 1);
+		int32 Y = FMath::Clamp((int32)(V * (Resolution - 1)), 0, Resolution - 1);
+		
+		float CenterDepth = Buffer[Y * Resolution + X];
+		float MaxDiff = 0.0f;
+		
+		// Check 3x3 neighborhood
+		for (int32 dy = -1; dy <= 1; dy++)
+		{
+			for (int32 dx = -1; dx <= 1; dx++)
+			{
+				if (dx == 0 && dy == 0) continue;
+				
+				int32 NX = FMath::Clamp(X + dx, 0, Resolution - 1);
+				int32 NY = FMath::Clamp(Y + dy, 0, Resolution - 1);
+				
+				float NeighborDepth = Buffer[NY * Resolution + NX];
+				MaxDiff = FMath::Max(MaxDiff, FMath::Abs(NeighborDepth - CenterDepth));
+			}
+		}
+		
+		return MaxDiff > Threshold;
+	};
+
+	int32 VerticesDisplaced = 0;
+	int32 VerticesSkippedBackfacing = 0;
+	int32 VerticesSkippedOutsideFrustum = 0;
+	int32 VerticesSkippedBehindCamera = 0;
+	int32 VerticesSkippedSmallDisplacement = 0;
+	float MaxDisplacementActual = 0.0f;
+	float MinDepthSampled = FLT_MAX;
+	float MaxDepthSampled = -FLT_MAX;
+
+	// Store displacements per vertex for smoothing pass
+	TMap<int32, float> VertexDisplacements;
+	TMap<int32, FVector3d> VertexWorldPositions;
+	TMap<int32, FVector3d> VertexWorldNormals;
+
+	for (int32 VID : DynMesh.VertexIndicesItr())
+	{
+		// Get vertex in MODEL space
+		FVector3d ModelPos = DynMesh.GetVertex(VID);
+		
+		// Transform to WORLD space using actor's transform
+		FVector3d WorldPos = FVector3d(ActorWorldTransform.TransformPosition(FVector(ModelPos)));
+
+		// Get pre-computed vertex normal and transform to world space
+		FVector3d ModelNormal = MeshNormals[VID];
+		FVector3d WorldNormal = FVector3d(ActorWorldTransform.TransformVectorNoScale(FVector(ModelNormal)));
+		WorldNormal.Normalize();
+
+		// Check if vertex normal faces toward camera (in world space)
+		// Use a negative threshold to allow vertices near the silhouette edge
+		// (their normals are perpendicular to view, so dot ≈ 0)
+		FVector3d ToCameraDir = (FVector3d(CameraLocation) - WorldPos).GetSafeNormal();
+		float NormalDotView = WorldNormal.Dot(ToCameraDir);
+
+		// Only skip vertices that clearly face away (threshold allows edge vertices)
+		if (NormalDotView < -0.1f)
+		{
+			VerticesSkippedBackfacing++;
+			continue;
+		}
+
+		// Project world-space vertex to UV coordinates
+		FVector2D UV = ProjectToUV(WorldPos);
+		
+		// Check for behind camera
+		if (UV.X < -0.5f)
+		{
+			VerticesSkippedBehindCamera++;
+			continue;
+		}
+		
+		// Check for outside frustum
+		if (UV.X < -1.5f)
+		{
+			VerticesSkippedOutsideFrustum++;
+			continue;
+		}
+		
+		float U = UV.X;
+		float V = UV.Y;
+
+		// Compute the vertex's actual depth (distance from camera along forward axis)
+		FVector ToVertex = FVector(WorldPos) - CameraLocation;
+		float VertexDepth = FVector::DotProduct(ToVertex, CamForward);
+		
+		// Normalize vertex depth to 0-1 range (same as depth buffer)
+		float NormalizedVertexDepth = (VertexDepth - CapturedMinDepth) / DepthRange;
+		NormalizedVertexDepth = FMath::Clamp(NormalizedVertexDepth, 0.0f, 1.0f);
+
+		// Sample depth buffers
+		float CurrentDepth = SampleBilinear(CurrentDepthBuffer, U, V);
+		float TargetDepth = SampleBilinear(TargetDepthBuffer, U, V);
+		
+		// Silhouette edge detection: if sampled depth doesn't match vertex depth,
+		// we're sampling across an edge (background vs foreground discontinuity)
+		float DepthMismatch = FMath::Abs(CurrentDepth - NormalizedVertexDepth);
+		if (DepthMismatch > 0.02f)  // 2% threshold - tighter
+		{
+			VerticesSkippedBackfacing++;  // Reusing this counter for edge vertices
+			continue;
+		}
+		
+		// Edge detection: skip vertices at depth discontinuities in either buffer
+		// This catches edges caused by resolution mismatch and bilinear interpolation artifacts
+		if (IsDepthEdge(CurrentDepthBuffer, U, V, 0.02f) || IsDepthEdge(TargetDepthBuffer, U, V, 0.02f))
+		{
+			VerticesSkippedOutsideFrustum++;  // Reusing counter for edge vertices
+			continue;
+		}
+		
+		MinDepthSampled = FMath::Min(MinDepthSampled, CurrentDepth);
+		MaxDepthSampled = FMath::Max(MaxDepthSampled, CurrentDepth);
+
+		// Calculate displacement
+		float Displacement = 0.0f;
+		
+		if (bUseAbsolute)
+		{
+			// ABSOLUTE MODE: Target image is a heightmap (standard depth convention)
+			// Darker (black/0) = closer to camera = pull toward
+			// Lighter (white/255) = further from camera = push away
+			// NormalizedCenter is the "zero displacement" point
+			float HeightValue = NormalizedCenter - TargetDepth;  // Flipped: dark=positive, light=negative
+			Displacement = HeightValue * DisplacementStrength * 2.0f;  // Scale to full strength range
+		}
+		else
+		{
+			// RELATIVE MODE: Compare target to current depth
+			// TargetDepth > CurrentDepth means target is farther (lighter) = push away
+			// TargetDepth < CurrentDepth means target is closer (darker) = pull toward camera
+			float DepthDiff = TargetDepth - CurrentDepth;
+			Displacement = DepthDiff * DisplacementStrength;
+		}
+		
+		// Fade out displacement based on how much the normal faces the camera
+		// NormalDotView: 1.0 = facing camera, 0.0 = perpendicular, <0 = away
+		// Use cubed falloff for even more aggressive fade at grazing angles
+		float FacingFactor = FMath::Clamp(NormalDotView, 0.0f, 1.0f);
+		FacingFactor = FacingFactor * FacingFactor * FacingFactor;  // Cubed for stronger falloff
+		Displacement *= FacingFactor;
+
+		if (bInvert)
+		{
+			Displacement = -Displacement;
+		}
+
+		// Skip tiny displacements
+		if (FMath::Abs(Displacement) < 0.001f)
+		{
+			VerticesSkippedSmallDisplacement++;
+			continue;
+		}
+
+		// Store displacement for smoothing pass (don't apply yet)
+		VertexDisplacements.Add(VID, Displacement);
+		VertexWorldPositions.Add(VID, WorldPos);
+		MaxDisplacementActual = FMath::Max(MaxDisplacementActual, FMath::Abs(Displacement));
+	}
+
+	// Smoothing pass: average each vertex's displacement with its neighbors
+	if (SmoothIterations > 0 && VertexDisplacements.Num() > 0)
+	{
+		for (int32 Iter = 0; Iter < SmoothIterations; Iter++)
+		{
+			TMap<int32, float> SmoothedDisplacements;
+			
+			for (auto& Pair : VertexDisplacements)
+			{
+				int32 VID = Pair.Key;
+				float OriginalDisp = Pair.Value;
+				
+				// Get neighbors via edge connectivity
+				float SumDisp = OriginalDisp;
+				int32 Count = 1;
+				
+				for (int32 EID : DynMesh.VtxEdgesItr(VID))
+				{
+					FIndex2i EdgeVerts = DynMesh.GetEdgeV(EID);
+					int32 NeighborVID = (EdgeVerts.A == VID) ? EdgeVerts.B : EdgeVerts.A;
+					
+					if (float* NeighborDisp = VertexDisplacements.Find(NeighborVID))
+					{
+						SumDisp += *NeighborDisp;
+						Count++;
+					}
+				}
+				
+				// Blend: 50% original, 50% neighbor average
+				float NeighborAvg = (Count > 1) ? (SumDisp - OriginalDisp) / (Count - 1) : OriginalDisp;
+				float SmoothedDisp = OriginalDisp * 0.5f + NeighborAvg * 0.5f;
+				SmoothedDisplacements.Add(VID, SmoothedDisp);
+			}
+			
+			VertexDisplacements = MoveTemp(SmoothedDisplacements);
+		}
+	}
+
+	// Apply smoothed displacements
+	for (auto& Pair : VertexDisplacements)
+	{
+		int32 VID = Pair.Key;
+		float Displacement = Pair.Value;
+		FVector3d WorldPos = VertexWorldPositions[VID];
+		
+		// Displace along camera forward direction (in world space)
+		FVector3d NewWorldPos = WorldPos + FVector3d(CamForward) * Displacement;
+		
+		// Transform back to model space for storage
+		FVector3d NewModelPos = FVector3d(ActorWorldTransform.InverseTransformPosition(FVector(NewWorldPos)));
+		
+		DynMesh.SetVertex(VID, NewModelPos);
+		VerticesDisplaced++;
+	}
+
+	FMeshNormals::QuickComputeVertexNormals(DynMesh);
+
+	UStaticMesh* OutputMesh = nullptr;
+	if (OutputPath == MeshPath)
+	{
+		MeshHelpers::SaveDynamicMeshToStaticMesh(DynMesh, SourceMesh);
+		OutputMesh = SourceMesh;
+	}
+	else
+	{
+		OutputMesh = MeshHelpers::CreateStaticMeshFromDynamicMesh(DynMesh, OutputPath, true, true, &Error);
+		if (!OutputMesh)
+		{
+			return FECACommandResult::Error(Error);
+		}
+	}
+
+	TSharedPtr<FJsonObject> Result = MeshHelpers::BuildMeshResultJson(OutputMesh, DynMesh.VertexCount(), DynMesh.TriangleCount());
+	Result->SetNumberField(TEXT("vertices_displaced"), VerticesDisplaced);
+	Result->SetNumberField(TEXT("vertices_skipped_backfacing"), VerticesSkippedBackfacing);
+	Result->SetNumberField(TEXT("vertices_skipped_outside_frustum"), VerticesSkippedOutsideFrustum);
+	Result->SetNumberField(TEXT("vertices_skipped_behind_camera"), VerticesSkippedBehindCamera);
+	Result->SetNumberField(TEXT("vertices_skipped_small_displacement"), VerticesSkippedSmallDisplacement);
+	Result->SetNumberField(TEXT("vertices_total"), DynMesh.VertexCount());
+	Result->SetNumberField(TEXT("max_displacement"), MaxDisplacementActual);
+	Result->SetNumberField(TEXT("min_depth_sampled"), MinDepthSampled);
+	Result->SetNumberField(TEXT("max_depth_sampled"), MaxDepthSampled);
+	Result->SetNumberField(TEXT("strength"), DisplacementStrength);
+	Result->SetNumberField(TEXT("resolution"), Resolution);
+	Result->SetNumberField(TEXT("current_depth_buffer_size"), CurrentDepthBuffer.Num());
+	Result->SetNumberField(TEXT("target_depth_buffer_size"), TargetDepthBuffer.Num());
+	Result->SetNumberField(TEXT("fov"), FOV);
+	Result->SetNumberField(TEXT("captured_min_depth_world"), CapturedMinDepth);
+	Result->SetNumberField(TEXT("captured_max_depth_world"), CapturedMaxDepth);
+	Result->SetNumberField(TEXT("depth_range_world"), DepthRange);
+
+	TSharedPtr<FJsonObject> CameraJson = MakeShared<FJsonObject>();
+	TSharedPtr<FJsonObject> LocationJson = MakeShared<FJsonObject>();
+	LocationJson->SetNumberField(TEXT("x"), CameraLocation.X);
+	LocationJson->SetNumberField(TEXT("y"), CameraLocation.Y);
+	LocationJson->SetNumberField(TEXT("z"), CameraLocation.Z);
+	CameraJson->SetObjectField(TEXT("location"), LocationJson);
+
+	TSharedPtr<FJsonObject> RotationJson = MakeShared<FJsonObject>();
+	RotationJson->SetNumberField(TEXT("pitch"), CameraRotation.Pitch);
+	RotationJson->SetNumberField(TEXT("yaw"), CameraRotation.Yaw);
+	RotationJson->SetNumberField(TEXT("roll"), CameraRotation.Roll);
+	CameraJson->SetObjectField(TEXT("rotation"), RotationJson);
+
+	CameraJson->SetNumberField(TEXT("fov"), FOV);
+	Result->SetObjectField(TEXT("camera"), CameraJson);
+
+	return FECACommandResult::Success(Result);
+}
+
+//==============================================================================
+// MESH DISPLACE FROM HEIGHTMAP (UV-based)
+//==============================================================================
+
+FECACommandResult FECACommand_MeshDisplaceFromHeightmap::Execute(const TSharedPtr<FJsonObject>& Params)
+{
+	// Get mesh path
+	FString MeshPath;
+	if (!GetStringParam(Params, TEXT("mesh_path"), MeshPath))
+	{
+		return FECACommandResult::Error(TEXT("Missing required parameter: mesh_path"));
+	}
+
+	// Get heightmap path
+	FString HeightmapPath;
+	if (!GetStringParam(Params, TEXT("heightmap_path"), HeightmapPath))
+	{
+		return FECACommandResult::Error(TEXT("Missing required parameter: heightmap_path"));
+	}
+
+	// Get output path (optional, defaults to overwriting input)
+	FString OutputPath = MeshPath;
+	GetStringParam(Params, TEXT("output_path"), OutputPath, false);
+
+	// Get parameters
+	double Strength = 100.0;
+	GetFloatParam(Params, TEXT("strength"), Strength, false);
+
+	int32 CenterValue = 0;
+	GetIntParam(Params, TEXT("center_value"), CenterValue, false);
+	float NormalizedCenter = CenterValue / 255.0f;
+
+	int32 UVChannel = 0;
+	GetIntParam(Params, TEXT("uv_channel"), UVChannel, false);
+
+	bool bInvert = false;
+	GetBoolParam(Params, TEXT("invert"), bInvert, false);
+
+	int32 SmoothIterations = 0;
+	GetIntParam(Params, TEXT("smooth_iterations"), SmoothIterations, false);
+
+	// Load the mesh
+	UStaticMesh* SourceMesh = MeshHelpers::LoadStaticMesh(MeshPath);
+	if (!SourceMesh)
+	{
+		return FECACommandResult::Error(FString::Printf(TEXT("Mesh not found: %s"), *MeshPath));
+	}
+
+	// Convert to dynamic mesh
+	FDynamicMesh3 DynMesh;
+	FString Error;
+	if (!MeshHelpers::LoadStaticMeshIntoDynamicMesh(SourceMesh, DynMesh, &Error))
+	{
+		return FECACommandResult::Error(Error);
+	}
+
+	// Check if mesh has UVs
+	if (!DynMesh.HasAttributes() || DynMesh.Attributes()->NumUVLayers() == 0)
+	{
+		return FECACommandResult::Error(TEXT("Mesh does not have UV coordinates"));
+	}
+
+	FDynamicMeshUVOverlay* UVOverlay = DynMesh.Attributes()->PrimaryUV();
+	if (!UVOverlay)
+	{
+		return FECACommandResult::Error(TEXT("Could not access UV overlay"));
+	}
+
+	// Load heightmap image
+	TArray<uint8> FileData;
+	if (!FFileHelper::LoadFileToArray(FileData, *HeightmapPath))
+	{
+		return FECACommandResult::Error(FString::Printf(TEXT("Failed to load heightmap: %s"), *HeightmapPath));
+	}
+
+	IImageWrapperModule& ImageWrapperModule = FModuleManager::LoadModuleChecked<IImageWrapperModule>(FName("ImageWrapper"));
+	TSharedPtr<IImageWrapper> ImageWrapper = ImageWrapperModule.CreateImageWrapper(EImageFormat::PNG);
+
+	if (!ImageWrapper->SetCompressed(FileData.GetData(), FileData.Num()))
+	{
+		return FECACommandResult::Error(TEXT("Failed to decode heightmap image"));
+	}
+
+	TArray<uint8> RawData;
+	if (!ImageWrapper->GetRaw(ERGBFormat::Gray, 8, RawData))
+	{
+		return FECACommandResult::Error(TEXT("Failed to get raw pixels from heightmap"));
+	}
+
+	int32 ImageWidth = ImageWrapper->GetWidth();
+	int32 ImageHeight = ImageWrapper->GetHeight();
+
+	// Convert to float buffer (0-1 range)
+	TArray<float> HeightBuffer;
+	HeightBuffer.SetNum(RawData.Num());
+	for (int32 i = 0; i < RawData.Num(); i++)
+	{
+		HeightBuffer[i] = RawData[i] / 255.0f;
+	}
+
+	// Bilinear sample helper
+	auto SampleHeightmap = [&](float U, float V) -> float
+	{
+		// Wrap UVs to 0-1 range
+		U = U - FMath::Floor(U);
+		V = V - FMath::Floor(V);
+
+		float PixelXf = U * (ImageWidth - 1);
+		float PixelYf = V * (ImageHeight - 1);
+
+		int32 X0 = FMath::Clamp((int32)FMath::Floor(PixelXf), 0, ImageWidth - 1);
+		int32 Y0 = FMath::Clamp((int32)FMath::Floor(PixelYf), 0, ImageHeight - 1);
+		int32 X1 = FMath::Min(X0 + 1, ImageWidth - 1);
+		int32 Y1 = FMath::Min(Y0 + 1, ImageHeight - 1);
+
+		float FracX = PixelXf - X0;
+		float FracY = PixelYf - Y0;
+
+		float V00 = HeightBuffer[Y0 * ImageWidth + X0];
+		float V10 = HeightBuffer[Y0 * ImageWidth + X1];
+		float V01 = HeightBuffer[Y1 * ImageWidth + X0];
+		float V11 = HeightBuffer[Y1 * ImageWidth + X1];
+
+		float Top = V00 + (V10 - V00) * FracX;
+		float Bottom = V01 + (V11 - V01) * FracX;
+		return Top + (Bottom - Top) * FracY;
+	};
+
+	// Compute vertex normals
+	FMeshNormals MeshNormals(&DynMesh);
+	MeshNormals.ComputeVertexNormals();
+
+	// Build a map of co-located vertices (vertices at the same position)
+	// This is needed to ensure vertices at UV seams get the same displacement
+	TMap<FVector3d, TArray<int32>> PositionToVertices;
+	for (int32 VID : DynMesh.VertexIndicesItr())
+	{
+		FVector3d Pos = DynMesh.GetVertex(VID);
+		// Quantize position to avoid floating point issues
+		FVector3d QuantizedPos(
+			FMath::RoundToFloat(Pos.X * 1000.0) / 1000.0,
+			FMath::RoundToFloat(Pos.Y * 1000.0) / 1000.0,
+			FMath::RoundToFloat(Pos.Z * 1000.0) / 1000.0
+		);
+		PositionToVertices.FindOrAdd(QuantizedPos).Add(VID);
+	}
+
+	// Store displacements per vertex
+	TMap<int32, float> VertexDisplacements;
+	TMap<int32, FVector3d> VertexNormals;
+	int32 VerticesProcessed = 0;
+	int32 VerticesSkippedNoUV = 0;
+	float MaxDisplacement = 0.0f;
+
+	// First pass: calculate displacements for each unique position
+	TSet<int32> ProcessedVertices;
+	
+	for (auto& Pair : PositionToVertices)
+	{
+		const TArray<int32>& ColocatedVerts = Pair.Value;
+		
+		// Gather all UVs and normals from co-located vertices
+		TArray<FVector2f> AllUVs;
+		FVector3d AvgNormal = FVector3d::Zero();
+		
+		for (int32 VID : ColocatedVerts)
+		{
+			// Accumulate normal
+			FVector3d Normal = MeshNormals[VID];
+			AvgNormal += Normal;
+			
+			// Get UVs from adjacent triangles
+			for (int32 TID : DynMesh.VtxTrianglesItr(VID))
+			{
+				FIndex3i TriVerts = DynMesh.GetTriangle(TID);
+				int32 LocalIndex = -1;
+				for (int32 i = 0; i < 3; i++)
+				{
+					if (TriVerts[i] == VID)
+					{
+						LocalIndex = i;
+						break;
+					}
+				}
+
+				if (LocalIndex >= 0)
+				{
+					FIndex3i UVTri = UVOverlay->GetTriangle(TID);
+					if (UVOverlay->IsElement(UVTri[LocalIndex]))
+					{
+						FVector2f UV = UVOverlay->GetElement(UVTri[LocalIndex]);
+						AllUVs.Add(UV);
+					}
+				}
+			}
+		}
+		
+		if (AllUVs.Num() == 0)
+		{
+			VerticesSkippedNoUV += ColocatedVerts.Num();
+			continue;
+		}
+		
+		// Average the UVs (handles UV seam vertices by averaging both sides)
+		FVector2f AvgUV = FVector2f::Zero();
+		for (const FVector2f& UV : AllUVs)
+		{
+			AvgUV += UV;
+		}
+		AvgUV /= AllUVs.Num();
+		
+		// Normalize the averaged normal
+		AvgNormal.Normalize();
+
+		// Sample heightmap
+		float HeightValue = SampleHeightmap(AvgUV.X, AvgUV.Y);
+
+		// Calculate displacement relative to center value
+		float Displacement = (HeightValue - NormalizedCenter) * Strength;
+
+		if (bInvert)
+		{
+			Displacement = -Displacement;
+		}
+
+		// Apply same displacement to ALL co-located vertices
+		for (int32 VID : ColocatedVerts)
+		{
+			VertexDisplacements.Add(VID, Displacement);
+			VertexNormals.Add(VID, AvgNormal);
+			ProcessedVertices.Add(VID);
+		}
+		
+		MaxDisplacement = FMath::Max(MaxDisplacement, FMath::Abs(Displacement));
+		VerticesProcessed += ColocatedVerts.Num();
+	}
+
+	// Smoothing pass
+	if (SmoothIterations > 0 && VertexDisplacements.Num() > 0)
+	{
+		for (int32 Iter = 0; Iter < SmoothIterations; Iter++)
+		{
+			TMap<int32, float> SmoothedDisplacements;
+
+			for (auto& Pair : VertexDisplacements)
+			{
+				int32 VID = Pair.Key;
+				float OriginalDisp = Pair.Value;
+
+				float SumDisp = OriginalDisp;
+				int32 Count = 1;
+
+				for (int32 EID : DynMesh.VtxEdgesItr(VID))
+				{
+					FIndex2i EdgeVerts = DynMesh.GetEdgeV(EID);
+					int32 NeighborVID = (EdgeVerts.A == VID) ? EdgeVerts.B : EdgeVerts.A;
+
+					if (float* NeighborDisp = VertexDisplacements.Find(NeighborVID))
+					{
+						SumDisp += *NeighborDisp;
+						Count++;
+					}
+				}
+
+				float NeighborAvg = (Count > 1) ? (SumDisp - OriginalDisp) / (Count - 1) : OriginalDisp;
+				float SmoothedDisp = OriginalDisp * 0.5f + NeighborAvg * 0.5f;
+				SmoothedDisplacements.Add(VID, SmoothedDisp);
+			}
+
+			VertexDisplacements = MoveTemp(SmoothedDisplacements);
+		}
+	}
+
+	// Apply displacements
+	int32 VerticesDisplaced = 0;
+	for (auto& Pair : VertexDisplacements)
+	{
+		int32 VID = Pair.Key;
+		float Displacement = Pair.Value;
+		FVector3d Normal = VertexNormals[VID];
+
+		FVector3d OldPos = DynMesh.GetVertex(VID);
+		FVector3d NewPos = OldPos + Normal * Displacement;
+		DynMesh.SetVertex(VID, NewPos);
+		VerticesDisplaced++;
+	}
+
+	// Recompute normals after displacement
+	FMeshNormals::QuickComputeVertexNormals(DynMesh);
+
+	// Save result
+	UStaticMesh* OutputMesh = nullptr;
+	if (OutputPath == MeshPath)
+	{
+		MeshHelpers::SaveDynamicMeshToStaticMesh(DynMesh, SourceMesh);
+		OutputMesh = SourceMesh;
+	}
+	else
+	{
+		OutputMesh = MeshHelpers::CreateStaticMeshFromDynamicMesh(DynMesh, OutputPath, true, true, &Error);
+		if (!OutputMesh)
+		{
+			return FECACommandResult::Error(Error);
+		}
+	}
+
+	// Build result
+	TSharedPtr<FJsonObject> Result = MeshHelpers::BuildMeshResultJson(OutputMesh, DynMesh.VertexCount(), DynMesh.TriangleCount());
+	Result->SetNumberField(TEXT("vertices_processed"), VerticesProcessed);
+	Result->SetNumberField(TEXT("vertices_displaced"), VerticesDisplaced);
+	Result->SetNumberField(TEXT("vertices_skipped_no_uv"), VerticesSkippedNoUV);
+	Result->SetNumberField(TEXT("max_displacement"), MaxDisplacement);
+	Result->SetNumberField(TEXT("strength"), Strength);
+	Result->SetNumberField(TEXT("center_value"), CenterValue);
+	Result->SetNumberField(TEXT("uv_channel"), UVChannel);
+	Result->SetNumberField(TEXT("smooth_iterations"), SmoothIterations);
+	Result->SetNumberField(TEXT("heightmap_width"), ImageWidth);
+	Result->SetNumberField(TEXT("heightmap_height"), ImageHeight);
+
+	return FECACommandResult::Success(Result);
+}
