ECABridge experimental MCP plugin

diff --git a/Engine/Plugins/Experimental/ECABridge/Source/ECABridge/Private/Commands/ECAMetasoundCommands.cpp b/Engine/Plugins/Experimental/ECABridge/Source/ECABridge/Private/Commands/ECAMetasoundCommands.cpp
new file mode 100644
--- /dev/null
+++ b/Engine/Plugins/Experimental/ECABridge/Source/ECABridge/Private/Commands/ECAMetasoundCommands.cpp
@@ -0,0 +1,2154 @@
+// Copyright Epic Games, Inc. All Rights Reserved.
+
+#include "Commands/ECAMetasoundCommands.h"
+#include "Commands/ECACommand.h"
+#include "Editor.h"
+#include "Engine/World.h"
+#include "EngineUtils.h"
+#include "AssetRegistry/AssetRegistryModule.h"
+
+// MetaSound includes
+#include "MetasoundSource.h"
+#include "MetasoundAssetBase.h"
+#include "MetasoundFrontendDocument.h"
+#include "MetasoundFrontendDocumentBuilder.h"
+#include "MetasoundFrontendSearchEngine.h"
+#include "MetasoundFrontendRegistryContainer.h"
+#include "MetasoundDocumentBuilderRegistry.h"
+#include "MetasoundUObjectRegistry.h"
+#include "MetasoundBuilderBase.h"
+#include "MetasoundBuilderSubsystem.h"
+
+// MetaSound Editor includes
+#if WITH_EDITOR
+#include "MetasoundEditorSubsystem.h"
+#endif
+
+
+
+#include "AssetToolsModule.h"
+#include "IAssetTools.h"
+#include "FileHelpers.h"
+#include "UObject/SavePackage.h"
+#include "Factories/Factory.h"
+#include "ScopedTransaction.h"
+
+// Audio component for spawning
+#include "Components/AudioComponent.h"
+#include "Sound/SoundBase.h"
+#include "Kismet/GameplayStatics.h"
+#include "GameFramework/Actor.h"
+
+// Register all Metasound commands
+REGISTER_ECA_COMMAND(FECACommand_GetMetasoundSources)
+REGISTER_ECA_COMMAND(FECACommand_CreateMetasoundSource)
+REGISTER_ECA_COMMAND(FECACommand_GetMetasoundNodes)
+REGISTER_ECA_COMMAND(FECACommand_AddMetasoundNode)
+REGISTER_ECA_COMMAND(FECACommand_RemoveMetasoundNode)
+REGISTER_ECA_COMMAND(FECACommand_ConnectMetasoundNodes)
+REGISTER_ECA_COMMAND(FECACommand_DisconnectMetasoundNodes)
+REGISTER_ECA_COMMAND(FECACommand_SetMetasoundNodeInput)
+REGISTER_ECA_COMMAND(FECACommand_AddMetasoundInput)
+REGISTER_ECA_COMMAND(FECACommand_RemoveMetasoundInput)
+REGISTER_ECA_COMMAND(FECACommand_GetMetasoundNodeTypes)
+REGISTER_ECA_COMMAND(FECACommand_PreviewMetasound)
+REGISTER_ECA_COMMAND(FECACommand_GetMetasoundInterface)
+REGISTER_ECA_COMMAND(FECACommand_SpawnMetasoundPlayer)
+REGISTER_ECA_COMMAND(FECACommand_AutoLayoutMetasoundGraph)
+
+//------------------------------------------------------------------------------
+// Helper Functions
+//------------------------------------------------------------------------------
+
+// Helper to get the UMetaSoundBuilderBase for an existing MetaSound asset using the Editor Subsystem
+static UMetaSoundBuilderBase* GetMetaSoundBuilder(UMetaSoundSource* MetaSound, EMetaSoundBuilderResult& OutResult)
+{
+	OutResult = EMetaSoundBuilderResult::Failed;
+	
+	if (!MetaSound)
+	{
+		return nullptr;
+	}
+
+#if WITH_EDITOR
+	UMetaSoundEditorSubsystem* EditorSubsystem = GEditor ? GEditor->GetEditorSubsystem<UMetaSoundEditorSubsystem>() : nullptr;
+	if (!EditorSubsystem)
+	{
+		return nullptr;
+	}
+	
+	// Get or create a builder for this asset
+	UMetaSoundBuilderBase* Builder = EditorSubsystem->FindOrBeginBuilding(MetaSound, OutResult);
+	return Builder;
+#else
+	return nullptr;
+#endif
+}
+
+// Helper to synchronize the MetaSound editor graph after document modifications
+static void SynchronizeMetaSoundEditorGraph(UMetaSoundSource* MetaSound, UMetaSoundBuilderBase* Builder = nullptr)
+{
+	if (!MetaSound)
+	{
+		return;
+	}
+
+#if WITH_EDITOR
+	// If we have a builder, conform the document to persist changes
+	if (Builder)
+	{
+		Builder->ConformObjectToDocument();
+	}
+	
+	UMetaSoundEditorSubsystem* EditorSubsystem = GEditor ? GEditor->GetEditorSubsystem<UMetaSoundEditorSubsystem>() : nullptr;
+	if (EditorSubsystem)
+	{
+		// Use the editor subsystem to properly register and synchronize
+		EditorSubsystem->RegisterGraphWithFrontend(*MetaSound, true /* bForceViewSynchronization */);
+	}
+#endif
+}
+
+// Constants for auto-layout spacing in MetaSound
+static const double MS_AUTO_LAYOUT_SPACING_X = 300.0;
+static const double MS_AUTO_LAYOUT_SPACING_Y = 100.0;
+
+// Calculate automatic position for a new MetaSound node
+static FVector2D CalculateAutoMetaSoundNodePosition(UMetaSoundSource* MetaSound, UMetaSoundBuilderBase* Builder)
+{
+	using namespace Metasound::Frontend;
+	
+	if (!MetaSound || !Builder)
+	{
+		return FVector2D(0, 0);
+	}
+	
+	const FMetaSoundFrontendDocumentBuilder& DocBuilder = Builder->GetConstBuilder();
+	
+	// Find the bounds of existing nodes
+	double MinX = TNumericLimits<double>::Max();
+	double MaxX = TNumericLimits<double>::Lowest();
+	double MinY = TNumericLimits<double>::Max();
+	double MaxY = TNumericLimits<double>::Lowest();
+	int32 NodeCount = 0;
+	
+	DocBuilder.IterateNodesByPredicate(
+		[&](const FMetasoundFrontendClass& NodeClass, const FMetasoundFrontendNode& Node, const FGuid& PageID)
+		{
+#if WITH_EDITORONLY_DATA
+			for (const auto& LocationPair : Node.Style.Display.Locations)
+			{
+				MinX = FMath::Min(MinX, (double)LocationPair.Value.X);
+				MaxX = FMath::Max(MaxX, (double)LocationPair.Value.X);
+				MinY = FMath::Min(MinY, (double)LocationPair.Value.Y);
+				MaxY = FMath::Max(MaxY, (double)LocationPair.Value.Y);
+				NodeCount++;
+				break; // Only need first location
+			}
+#endif
+		},
+		[](const FMetasoundFrontendNode& Node) { return true; },
+		nullptr,
+		true
+	);
+	
+	// If no nodes exist, start at origin
+	if (NodeCount == 0)
+	{
+		return FVector2D(0, 0);
+	}
+	
+	// Place new node to the right of existing nodes (MetaSound flows left-to-right)
+	return FVector2D(MaxX + MS_AUTO_LAYOUT_SPACING_X, MinY);
+}
+
+static UMetaSoundSource* LoadMetaSoundSource(const FString& AssetPath)
+{
+	FString FullPath = AssetPath;
+	if (!FullPath.EndsWith(TEXT(".")) && !FullPath.Contains(TEXT(".")))
+	{
+		FullPath = AssetPath + TEXT(".") + FPackageName::GetShortName(AssetPath);
+	}
+	
+	UMetaSoundSource* MetaSound = LoadObject<UMetaSoundSource>(nullptr, *AssetPath);
+	if (!MetaSound)
+	{
+		FString CleanPath = AssetPath;
+		CleanPath.RemoveFromEnd(TEXT("_C"));
+		MetaSound = LoadObject<UMetaSoundSource>(nullptr, *CleanPath);
+	}
+	
+	return MetaSound;
+}
+
+static FString GetNodeIdString(const FGuid& NodeId)
+{
+	return NodeId.ToString(EGuidFormats::DigitsWithHyphens);
+}
+
+static bool ParseNodeId(const FString& NodeIdStr, FGuid& OutGuid)
+{
+	return FGuid::Parse(NodeIdStr, OutGuid);
+}
+
+static TSharedPtr<FJsonObject> NodeToJson(const FMetasoundFrontendNode& Node, bool bIncludeDetails)
+{
+	TSharedPtr<FJsonObject> NodeObj = MakeShared<FJsonObject>();
+	NodeObj->SetStringField(TEXT("id"), GetNodeIdString(Node.GetID()));
+	NodeObj->SetStringField(TEXT("name"), Node.Name.ToString());
+	NodeObj->SetStringField(TEXT("class_id"), GetNodeIdString(Node.ClassID));
+	
+	// Position
+	TSharedPtr<FJsonObject> PosObj = MakeShared<FJsonObject>();
+#if WITH_EDITORONLY_DATA
+	if (Node.Style.Display.Locations.Num() > 0)
+	{
+		// Get first location entry
+		for (const auto& LocationPair : Node.Style.Display.Locations)
+		{
+			PosObj->SetNumberField(TEXT("x"), LocationPair.Value.X);
+			PosObj->SetNumberField(TEXT("y"), LocationPair.Value.Y);
+			break;
+		}
+	}
+	else
+#endif
+	{
+		PosObj->SetNumberField(TEXT("x"), 0.0);
+		PosObj->SetNumberField(TEXT("y"), 0.0);
+	}
+	NodeObj->SetObjectField(TEXT("position"), PosObj);
+	
+	if (bIncludeDetails)
+	{
+		// Input pins
+		TArray<TSharedPtr<FJsonValue>> InputsArray;
+		for (const FMetasoundFrontendVertex& Input : Node.Interface.Inputs)
+		{
+			TSharedPtr<FJsonObject> InputObj = MakeShared<FJsonObject>();
+			InputObj->SetStringField(TEXT("name"), Input.Name.ToString());
+			InputObj->SetStringField(TEXT("type"), Input.TypeName.ToString());
+			InputObj->SetStringField(TEXT("id"), GetNodeIdString(Input.VertexID));
+			InputsArray.Add(MakeShared<FJsonValueObject>(InputObj));
+		}
+		NodeObj->SetArrayField(TEXT("inputs"), InputsArray);
+		
+		// Output pins
+		TArray<TSharedPtr<FJsonValue>> OutputsArray;
+		for (const FMetasoundFrontendVertex& Output : Node.Interface.Outputs)
+		{
+			TSharedPtr<FJsonObject> OutputObj = MakeShared<FJsonObject>();
+			OutputObj->SetStringField(TEXT("name"), Output.Name.ToString());
+			OutputObj->SetStringField(TEXT("type"), Output.TypeName.ToString());
+			OutputObj->SetStringField(TEXT("id"), GetNodeIdString(Output.VertexID));
+			OutputsArray.Add(MakeShared<FJsonValueObject>(OutputObj));
+		}
+		NodeObj->SetArrayField(TEXT("outputs"), OutputsArray);
+	}
+	
+	return NodeObj;
+}
+
+//------------------------------------------------------------------------------
+// GetMetasoundSources
+//------------------------------------------------------------------------------
+
+FECACommandResult FECACommand_GetMetasoundSources::Execute(const TSharedPtr<FJsonObject>& Params)
+{
+	FString PathFilter;
+	FString NameFilter;
+	GetStringParam(Params, TEXT("path_filter"), PathFilter, false);
+	GetStringParam(Params, TEXT("name_filter"), NameFilter, false);
+	
+	FAssetRegistryModule& AssetRegistryModule = FModuleManager::LoadModuleChecked<FAssetRegistryModule>("AssetRegistry");
+	IAssetRegistry& AssetRegistry = AssetRegistryModule.Get();
+	
+	TArray<FAssetData> AssetDataList;
+	AssetRegistry.GetAssetsByClass(UMetaSoundSource::StaticClass()->GetClassPathName(), AssetDataList);
+	
+	TArray<TSharedPtr<FJsonValue>> SourcesArray;
+	
+	for (const FAssetData& AssetData : AssetDataList)
+	{
+		FString AssetPath = AssetData.GetObjectPathString();
+		FString AssetName = AssetData.AssetName.ToString();
+		
+		if (!PathFilter.IsEmpty() && !AssetPath.Contains(PathFilter))
+		{
+			continue;
+		}
+		if (!NameFilter.IsEmpty() && !AssetName.Contains(NameFilter))
+		{
+			continue;
+		}
+		
+		TSharedPtr<FJsonObject> SourceObj = MakeShared<FJsonObject>();
+		SourceObj->SetStringField(TEXT("name"), AssetName);
+		SourceObj->SetStringField(TEXT("path"), AssetData.PackageName.ToString());
+		SourceObj->SetStringField(TEXT("object_path"), AssetPath);
+		
+		SourcesArray.Add(MakeShared<FJsonValueObject>(SourceObj));
+	}
+	
+	TSharedPtr<FJsonObject> Result = MakeResult();
+	Result->SetArrayField(TEXT("sources"), SourcesArray);
+	Result->SetNumberField(TEXT("count"), SourcesArray.Num());
+	
+	return FECACommandResult::Success(Result);
+}
+
+//------------------------------------------------------------------------------
+// CreateMetasoundSource
+//------------------------------------------------------------------------------
+
+FECACommandResult FECACommand_CreateMetasoundSource::Execute(const TSharedPtr<FJsonObject>& Params)
+{
+	FString AssetPath;
+	FString OutputFormat = TEXT("Stereo");
+	
+	if (!GetStringParam(Params, TEXT("asset_path"), AssetPath, true))
+	{
+		return FECACommandResult::Error(TEXT("asset_path is required"));
+	}
+	GetStringParam(Params, TEXT("output_format"), OutputFormat, false);
+	
+	FString PackagePath = FPackageName::GetLongPackagePath(AssetPath);
+	FString AssetName = FPackageName::GetShortName(AssetPath);
+	
+	FString FullPackagePath = PackagePath / AssetName;
+	
+	// Check if asset already exists at this path
+	if (UObject* ExistingAsset = StaticLoadObject(UObject::StaticClass(), nullptr, *FullPackagePath))
+	{
+		// Asset already exists - return info about it instead of crashing
+		if (UMetaSoundSource* ExistingMetaSound = Cast<UMetaSoundSource>(ExistingAsset))
+		{
+			TSharedPtr<FJsonObject> Result = MakeResult();
+			Result->SetStringField(TEXT("asset_path"), FullPackagePath);
+			Result->SetStringField(TEXT("name"), AssetName);
+			Result->SetBoolField(TEXT("already_exists"), true);
+			Result->SetStringField(TEXT("message"), TEXT("MetaSound source already exists at this path"));
+			return FECACommandResult::Success(Result);
+		}
+		else
+		{
+			return FECACommandResult::Error(FString::Printf(TEXT("Asset already exists at path '%s' but is not a MetaSound source (it's a %s)"), *FullPackagePath, *ExistingAsset->GetClass()->GetName()));
+		}
+	}
+	
+	UPackage* Package = CreatePackage(*FullPackagePath);
+	if (!Package)
+	{
+		return FECACommandResult::Error(FString::Printf(TEXT("Failed to create package: %s"), *FullPackagePath));
+	}
+	
+	UMetaSoundSource* NewMetaSound = NewObject<UMetaSoundSource>(Package, *AssetName, RF_Public | RF_Standalone);
+	if (!NewMetaSound)
+	{
+		return FECACommandResult::Error(TEXT("Failed to create MetaSound source object"));
+	}
+	
+	Package->MarkPackageDirty();
+	FAssetRegistryModule::AssetCreated(NewMetaSound);
+	
+	FString PackageFilename = FPackageName::LongPackageNameToFilename(FullPackagePath, FPackageName::GetAssetPackageExtension());
+	FSavePackageArgs SaveArgs;
+	SaveArgs.TopLevelFlags = RF_Public | RF_Standalone;
+	UPackage::SavePackage(Package, NewMetaSound, *PackageFilename, SaveArgs);
+	
+	TSharedPtr<FJsonObject> Result = MakeResult();
+	Result->SetStringField(TEXT("asset_path"), FullPackagePath);
+	Result->SetStringField(TEXT("name"), AssetName);
+	Result->SetStringField(TEXT("output_format"), OutputFormat);
+	
+	return FECACommandResult::Success(Result);
+}
+
+//------------------------------------------------------------------------------
+// GetMetasoundNodes
+//------------------------------------------------------------------------------
+
+FECACommandResult FECACommand_GetMetasoundNodes::Execute(const TSharedPtr<FJsonObject>& Params)
+{
+	FString AssetPath;
+	bool bIncludeDetails = false;
+	
+	if (!GetStringParam(Params, TEXT("asset_path"), AssetPath, true))
+	{
+		return FECACommandResult::Error(TEXT("asset_path is required"));
+	}
+	GetBoolParam(Params, TEXT("include_details"), bIncludeDetails, false);
+	
+	UMetaSoundSource* MetaSound = LoadMetaSoundSource(AssetPath);
+	if (!MetaSound)
+	{
+		return FECACommandResult::Error(FString::Printf(TEXT("MetaSound not found: %s"), *AssetPath));
+	}
+	
+	// Use the Builder API to access nodes - this is the proper way to access MetaSound data
+	// and handles paged graphs correctly
+	EMetaSoundBuilderResult BuilderResult;
+	UMetaSoundBuilderBase* Builder = GetMetaSoundBuilder(MetaSound, BuilderResult);
+	if (!Builder || BuilderResult != EMetaSoundBuilderResult::Succeeded)
+	{
+		return FECACommandResult::Error(TEXT("Failed to get MetaSound builder. Make sure the editor is running."));
+	}
+	
+	const FMetaSoundFrontendDocumentBuilder& DocBuilder = Builder->GetConstBuilder();
+	
+	TArray<TSharedPtr<FJsonValue>> NodesArray;
+	TArray<TSharedPtr<FJsonValue>> ConnectionsArray;
+	
+	// Use the Builder's IterateNodesByPredicate method to iterate through ALL nodes across ALL pages
+	// The predicate always returns true to include all nodes, and bIterateAllPages=true ensures we get nodes from all graph pages
+	DocBuilder.IterateNodesByPredicate(
+		[&](const FMetasoundFrontendClass& NodeClass, const FMetasoundFrontendNode& Node, const FGuid& PageID)
+		{
+			NodesArray.Add(MakeShared<FJsonValueObject>(NodeToJson(Node, bIncludeDetails)));
+		},
+		[](const FMetasoundFrontendNode& Node) { return true; }, // Include all nodes
+		nullptr, // No specific page ID
+		true     // bIterateAllPages = true - iterate ALL pages
+	);
+	
+	// For edges/connections, we still need to iterate through graph pages
+	// as the Builder API doesn't have a direct edge iteration method
+	const FMetasoundFrontendDocument& Document = DocBuilder.GetConstDocumentChecked();
+	const FMetasoundFrontendGraphClass& GraphClass = Document.RootGraph;
+	
+	GraphClass.IterateGraphPages([&](const FMetasoundFrontendGraph& Graph)
+	{
+		for (const FMetasoundFrontendEdge& Edge : Graph.Edges)
+		{
+			TSharedPtr<FJsonObject> EdgeObj = MakeShared<FJsonObject>();
+			EdgeObj->SetStringField(TEXT("from_node_id"), GetNodeIdString(Edge.FromNodeID));
+			EdgeObj->SetStringField(TEXT("from_vertex_id"), GetNodeIdString(Edge.FromVertexID));
+			EdgeObj->SetStringField(TEXT("to_node_id"), GetNodeIdString(Edge.ToNodeID));
+			EdgeObj->SetStringField(TEXT("to_vertex_id"), GetNodeIdString(Edge.ToVertexID));
+			ConnectionsArray.Add(MakeShared<FJsonValueObject>(EdgeObj));
+		}
+	});
+	
+	TSharedPtr<FJsonObject> Result = MakeResult();
+	Result->SetArrayField(TEXT("nodes"), NodesArray);
+	Result->SetNumberField(TEXT("node_count"), NodesArray.Num());
+	Result->SetArrayField(TEXT("connections"), ConnectionsArray);
+	Result->SetNumberField(TEXT("connection_count"), ConnectionsArray.Num());
+	
+	return FECACommandResult::Success(Result);
+}
+
+//------------------------------------------------------------------------------
+// AddMetasoundNode
+//------------------------------------------------------------------------------
+
+FECACommandResult FECACommand_AddMetasoundNode::Execute(const TSharedPtr<FJsonObject>& Params)
+{
+	using namespace Metasound::Frontend;
+	
+	FString AssetPath;
+	FString NodeClass;
+	FString NodeName;
+	double PositionX = 0.0;
+	double PositionY = 0.0;
+	
+	if (!GetStringParam(Params, TEXT("asset_path"), AssetPath, true))
+	{
+		return FECACommandResult::Error(TEXT("asset_path is required"));
+	}
+	if (!GetStringParam(Params, TEXT("node_class"), NodeClass, true))
+	{
+		return FECACommandResult::Error(TEXT("node_class is required"));
+	}
+	GetStringParam(Params, TEXT("node_name"), NodeName, false);
+	GetFloatParam(Params, TEXT("position_x"), PositionX, false);
+	GetFloatParam(Params, TEXT("position_y"), PositionY, false);
+	
+	UMetaSoundSource* MetaSound = LoadMetaSoundSource(AssetPath);
+	if (!MetaSound)
+	{
+		return FECACommandResult::Error(FString::Printf(TEXT("MetaSound not found: %s"), *AssetPath));
+	}
+	
+	// Search for the node class using the search engine
+	ISearchEngine& SearchEngine = ISearchEngine::Get();
+	
+	FMetaSoundClassInfo FoundClassInfo;
+	bool bFoundClass = false;
+	FString FoundDisplayName;
+	FMetasoundFrontendClassName FoundClassName;
+	
+#if WITH_EDITORONLY_DATA
+	TArray<FMetaSoundClassInfo> AllClasses = SearchEngine.FindAllClasses(ISearchEngine::EResultVersion::Highest, false);
+	
+	for (const FMetaSoundClassInfo& ClassInfo : AllClasses)
+	{
+		FString ClassName = ClassInfo.ClassName.Name.ToString();
+		FString Namespace = ClassInfo.ClassName.Namespace.ToString();
+		FString Variant = ClassInfo.ClassName.Variant.ToString();
+		FString FullClassName = FString::Printf(TEXT("%s.%s.%s"), *Namespace, *ClassName, *Variant);
+		FString DisplayName = ClassInfo.InterfaceInfo.SearchInfo.ClassDisplayName.ToString();
+		
+		// Case-insensitive search on class name, display name, or full class name
+		if (ClassName.Contains(NodeClass, ESearchCase::IgnoreCase) || 
+			DisplayName.Contains(NodeClass, ESearchCase::IgnoreCase) ||
+			FullClassName.Contains(NodeClass, ESearchCase::IgnoreCase))
+		{
+			// Verify the class is actually registered (not just in search index)
+			FMetaSoundClassInfo RegisteredClassInfo;
+			if (SearchEngine.FindRegisteredClass(ClassInfo.ClassName, ClassInfo.Version.Major, RegisteredClassInfo))
+			{
+				FoundClassInfo = RegisteredClassInfo;
+				FoundClassName = RegisteredClassInfo.ClassName;
+				FoundDisplayName = DisplayName.IsEmpty() ? ClassName : DisplayName;
+				bFoundClass = true;
+				break;
+			}
+		}
+	}
+#endif
+	
+	if (!bFoundClass)
+	{
+		// Try to provide more helpful error message
+		FString FoundButNotRegistered;
+#if WITH_EDITORONLY_DATA
+		for (const FMetaSoundClassInfo& ClassInfo : AllClasses)
+		{
+			FString ClassName = ClassInfo.ClassName.Name.ToString();
+			FString DisplayName = ClassInfo.InterfaceInfo.SearchInfo.ClassDisplayName.ToString();
+			if (ClassName.Contains(NodeClass, ESearchCase::IgnoreCase) || 
+				DisplayName.Contains(NodeClass, ESearchCase::IgnoreCase))
+			{
+				FoundButNotRegistered = FString::Printf(TEXT("%s.%s.%s (v%d.%d)"), 
+					*ClassInfo.ClassName.Namespace.ToString(),
+					*ClassInfo.ClassName.Name.ToString(),
+					*ClassInfo.ClassName.Variant.ToString(),
+					ClassInfo.Version.Major, ClassInfo.Version.Minor);
+				break;
+			}
+		}
+#endif
+		if (!FoundButNotRegistered.IsEmpty())
+		{
+			return FECACommandResult::Error(FString::Printf(TEXT("Node class '%s' found in search index as '%s' but not registered. This may be a native node that isn't loaded yet."), *NodeClass, *FoundButNotRegistered));
+		}
+		return FECACommandResult::Error(FString::Printf(TEXT("Node class not found: %s"), *NodeClass));
+	}
+	
+	FScopedTransaction Transaction(FText::FromString(TEXT("Add MetaSound Node")));
+	MetaSound->Modify();
+	
+	// Get the builder using the Editor Subsystem (proper high-level API)
+	EMetaSoundBuilderResult BuilderResult;
+	UMetaSoundBuilderBase* Builder = GetMetaSoundBuilder(MetaSound, BuilderResult);
+	if (!Builder || BuilderResult != EMetaSoundBuilderResult::Succeeded)
+	{
+		return FECACommandResult::Error(TEXT("Failed to get MetaSound builder. Make sure the editor is running."));
+	}
+	
+	// Use the high-level builder API to add the node
+	EMetaSoundBuilderResult AddResult;
+	FMetaSoundNodeHandle NewNodeHandle = Builder->AddNodeByClassName(FoundClassName, AddResult, FoundClassInfo.Version.Major);
+	
+	if (AddResult != EMetaSoundBuilderResult::Succeeded || !NewNodeHandle.NodeID.IsValid())
+	{
+		// Provide more detailed error info
+		FString FullClassName = FString::Printf(TEXT("%s.%s.%s"), 
+			*FoundClassName.Namespace.ToString(), 
+			*FoundClassName.Name.ToString(),
+			*FoundClassName.Variant.ToString());
+		return FECACommandResult::Error(FString::Printf(TEXT("Failed to add node of class: %s (searched: %s, found: %s, version: %d.%d)"), 
+			*NodeClass, *NodeClass, *FullClassName, FoundClassInfo.Version.Major, FoundClassInfo.Version.Minor));
+	}
+	
+	// IMPORTANT: Initialize default values for all node inputs from the class defaults.
+	// The FMetasoundFrontendNode constructor does NOT populate InputLiterals, so we must
+	// do it manually to prevent crashes when building/playing the MetaSound.
+	{
+		EMetaSoundBuilderResult InputsResult;
+		TArray<FMetaSoundBuilderNodeInputHandle> NodeInputs = Builder->FindNodeInputs(NewNodeHandle, InputsResult);
+		if (InputsResult == EMetaSoundBuilderResult::Succeeded)
+		{
+			for (const FMetaSoundBuilderNodeInputHandle& InputHandle : NodeInputs)
+			{
+				// Get the class default value for this input
+				EMetaSoundBuilderResult DefaultResult;
+				FMetasoundFrontendLiteral ClassDefault = Builder->GetNodeInputClassDefault(InputHandle, DefaultResult);
+				if (DefaultResult == EMetaSoundBuilderResult::Succeeded)
+				{
+					// Set the default value on the node's input
+					Builder->SetNodeInputDefault(InputHandle, ClassDefault, DefaultResult);
+				}
+			}
+		}
+	}
+	
+	// Get the node ID for the result
+	FGuid NewNodeId = NewNodeHandle.NodeID;
+	
+#if WITH_EDITOR
+	// Set the node position using the Editor Subsystem
+	{
+		UMetaSoundEditorSubsystem* EditorSubsystem = GEditor ? GEditor->GetEditorSubsystem<UMetaSoundEditorSubsystem>() : nullptr;
+		if (EditorSubsystem)
+		{
+			FVector2D NodePosition;
+			
+			// Use provided position or auto-calculate
+			if (PositionX != 0.0 || PositionY != 0.0)
+			{
+				NodePosition = FVector2D(PositionX, PositionY);
+			}
+			else
+			{
+				// Auto-calculate position for new node
+				NodePosition = CalculateAutoMetaSoundNodePosition(MetaSound, Builder);
+			}
+			
+			EMetaSoundBuilderResult LocationResult;
+			EditorSubsystem->SetNodeLocation(Builder, NewNodeHandle, NodePosition, LocationResult);
+		}
+	}
+#endif
+	
+	MetaSound->MarkPackageDirty();
+	
+	// Synchronize the editor graph to reflect the changes in the UI
+	SynchronizeMetaSoundEditorGraph(MetaSound, Builder);
+	
+	TSharedPtr<FJsonObject> Result = MakeResult();
+	Result->SetStringField(TEXT("node_id"), GetNodeIdString(NewNodeId));
+	Result->SetStringField(TEXT("node_name"), NodeName.IsEmpty() ? FoundClassName.Name.ToString() : NodeName);
+	Result->SetStringField(TEXT("class_name"), FoundDisplayName);
+	
+	return FECACommandResult::Success(Result);
+}
+
+//------------------------------------------------------------------------------
+// RemoveMetasoundNode
+//------------------------------------------------------------------------------
+
+FECACommandResult FECACommand_RemoveMetasoundNode::Execute(const TSharedPtr<FJsonObject>& Params)
+{
+	using namespace Metasound::Frontend;
+	
+	FString AssetPath;
+	FString NodeIdStr;
+	
+	if (!GetStringParam(Params, TEXT("asset_path"), AssetPath, true))
+	{
+		return FECACommandResult::Error(TEXT("asset_path is required"));
+	}
+	if (!GetStringParam(Params, TEXT("node_id"), NodeIdStr, true))
+	{
+		return FECACommandResult::Error(TEXT("node_id is required"));
+	}
+	
+	FGuid NodeId;
+	if (!ParseNodeId(NodeIdStr, NodeId))
+	{
+		return FECACommandResult::Error(FString::Printf(TEXT("Invalid node_id format: %s"), *NodeIdStr));
+	}
+	
+	UMetaSoundSource* MetaSound = LoadMetaSoundSource(AssetPath);
+	if (!MetaSound)
+	{
+		return FECACommandResult::Error(FString::Printf(TEXT("MetaSound not found: %s"), *AssetPath));
+	}
+	
+	FScopedTransaction Transaction(FText::FromString(TEXT("Remove MetaSound Node")));
+	MetaSound->Modify();
+	
+	// Get the builder using the Editor Subsystem (proper high-level API)
+	EMetaSoundBuilderResult BuilderResult;
+	UMetaSoundBuilderBase* Builder = GetMetaSoundBuilder(MetaSound, BuilderResult);
+	if (!Builder || BuilderResult != EMetaSoundBuilderResult::Succeeded)
+	{
+		return FECACommandResult::Error(TEXT("Failed to get MetaSound builder. Make sure the editor is running."));
+	}
+	
+	// Create a node handle from the GUID
+	FMetaSoundNodeHandle NodeHandle;
+	NodeHandle.NodeID = NodeId;
+	
+	// Check if the node exists
+	if (!Builder->ContainsNode(NodeHandle))
+	{
+		return FECACommandResult::Error(FString::Printf(TEXT("Node not found: %s"), *NodeIdStr));
+	}
+	
+	// Remove the node using the builder (this also removes associated connections)
+	EMetaSoundBuilderResult RemoveResult;
+	Builder->RemoveNode(NodeHandle, RemoveResult, true /* bRemoveUnusedDependencies */);
+	
+	if (RemoveResult != EMetaSoundBuilderResult::Succeeded)
+	{
+		return FECACommandResult::Error(FString::Printf(TEXT("Failed to remove node: %s"), *NodeIdStr));
+	}
+	
+	MetaSound->MarkPackageDirty();
+	
+	// Synchronize the editor graph to reflect the changes in the UI
+	SynchronizeMetaSoundEditorGraph(MetaSound, Builder);
+	
+	TSharedPtr<FJsonObject> Result = MakeResult();
+	Result->SetStringField(TEXT("removed_node_id"), NodeIdStr);
+	
+	return FECACommandResult::Success(Result);
+}
+
+//------------------------------------------------------------------------------
+// ConnectMetasoundNodes
+//------------------------------------------------------------------------------
+
+FECACommandResult FECACommand_ConnectMetasoundNodes::Execute(const TSharedPtr<FJsonObject>& Params)
+{
+	FString AssetPath;
+	FString SourceNodeIdStr;
+	FString SourcePin;
+	FString TargetNodeIdStr;
+	FString TargetPin;
+	
+	if (!GetStringParam(Params, TEXT("asset_path"), AssetPath, true))
+	{
+		return FECACommandResult::Error(TEXT("asset_path is required"));
+	}
+	if (!GetStringParam(Params, TEXT("source_node_id"), SourceNodeIdStr, true))
+	{
+		return FECACommandResult::Error(TEXT("source_node_id is required"));
+	}
+	if (!GetStringParam(Params, TEXT("source_pin"), SourcePin, true))
+	{
+		return FECACommandResult::Error(TEXT("source_pin is required"));
+	}
+	if (!GetStringParam(Params, TEXT("target_node_id"), TargetNodeIdStr, true))
+	{
+		return FECACommandResult::Error(TEXT("target_node_id is required"));
+	}
+	if (!GetStringParam(Params, TEXT("target_pin"), TargetPin, true))
+	{
+		return FECACommandResult::Error(TEXT("target_pin is required"));
+	}
+	
+	FGuid SourceNodeId, TargetNodeId;
+	if (!ParseNodeId(SourceNodeIdStr, SourceNodeId))
+	{
+		return FECACommandResult::Error(FString::Printf(TEXT("Invalid source_node_id format: %s"), *SourceNodeIdStr));
+	}
+	if (!ParseNodeId(TargetNodeIdStr, TargetNodeId))
+	{
+		return FECACommandResult::Error(FString::Printf(TEXT("Invalid target_node_id format: %s"), *TargetNodeIdStr));
+	}
+	
+	UMetaSoundSource* MetaSound = LoadMetaSoundSource(AssetPath);
+	if (!MetaSound)
+	{
+		return FECACommandResult::Error(FString::Printf(TEXT("MetaSound not found: %s"), *AssetPath));
+	}
+	
+	FScopedTransaction Transaction(FText::FromString(TEXT("Connect MetaSound Nodes")));
+	MetaSound->Modify();
+	
+	// Get the builder using the Editor Subsystem (proper high-level API)
+	EMetaSoundBuilderResult BuilderResult;
+	UMetaSoundBuilderBase* Builder = GetMetaSoundBuilder(MetaSound, BuilderResult);
+	if (!Builder || BuilderResult != EMetaSoundBuilderResult::Succeeded)
+	{
+		return FECACommandResult::Error(TEXT("Failed to get MetaSound builder. Make sure the editor is running."));
+	}
+	
+	// Use the Builder API to find node outputs and inputs by name
+	FMetaSoundNodeHandle SourceNodeHandle(SourceNodeId);
+	FMetaSoundNodeHandle TargetNodeHandle(TargetNodeId);
+	
+	FMetaSoundBuilderNodeOutputHandle OutputHandle;
+	FMetaSoundBuilderNodeInputHandle InputHandle;
+	
+	// First, try to find the source as a graph input node (they have special handling)
+	// Graph input nodes can be connected using their name as the "pin"
+	FMetaSoundNodeHandle GraphInputNodeHandle = Builder->FindGraphInputNode(FName(*SourcePin), BuilderResult);
+	if (BuilderResult == EMetaSoundBuilderResult::Succeeded && GraphInputNodeHandle.IsSet())
+	{
+		// This is a graph input - use ConnectGraphInputToNode helper
+		Builder->ConnectGraphInputToNode(FName(*SourcePin), TargetNodeHandle, FName(*TargetPin), BuilderResult);
+		if (BuilderResult != EMetaSoundBuilderResult::Succeeded)
+		{
+			return FECACommandResult::Error(FString::Printf(TEXT("Failed to connect graph input '%s' to node input '%s'"), *SourcePin, *TargetPin));
+		}
+		
+		MetaSound->MarkPackageDirty();
+		SynchronizeMetaSoundEditorGraph(MetaSound, Builder);
+		
+		TSharedPtr<FJsonObject> Result = MakeResult();
+		Result->SetStringField(TEXT("source_node_id"), SourceNodeIdStr);
+		Result->SetStringField(TEXT("source_pin"), SourcePin);
+		Result->SetStringField(TEXT("target_node_id"), TargetNodeIdStr);
+		Result->SetStringField(TEXT("target_pin"), TargetPin);
+		Result->SetStringField(TEXT("connection_type"), TEXT("graph_input_to_node"));
+		
+		return FECACommandResult::Success(Result);
+	}
+	
+	// Standard case: Find the source output pin using Builder API
+	OutputHandle = Builder->FindNodeOutputByName(SourceNodeHandle, FName(*SourcePin), BuilderResult);
+	if (BuilderResult != EMetaSoundBuilderResult::Succeeded)
+	{
+		// If not found by name, try finding all outputs and use the first one (for nodes with single output)
+		TArray<FMetaSoundBuilderNodeOutputHandle> AllOutputs = Builder->FindNodeOutputs(SourceNodeHandle, BuilderResult);
+		if (BuilderResult == EMetaSoundBuilderResult::Succeeded && AllOutputs.Num() > 0)
+		{
+			// Try to match by index if the pin name looks like an index
+			int32 PinIndex = 0;
+			if (SourcePin.IsNumeric())
+			{
+				PinIndex = FCString::Atoi(*SourcePin);
+			}
+			if (PinIndex >= 0 && PinIndex < AllOutputs.Num())
+			{
+				OutputHandle = AllOutputs[PinIndex];
+			}
+			else if (AllOutputs.Num() == 1)
+			{
+				// Single output node - use the only output
+				OutputHandle = AllOutputs[0];
+			}
+			else
+			{
+				return FECACommandResult::Error(FString::Printf(TEXT("Source output pin '%s' not found on node '%s'. Available outputs: %d"), *SourcePin, *SourceNodeIdStr, AllOutputs.Num()));
+			}
+		}
+		else
+		{
+			return FECACommandResult::Error(FString::Printf(TEXT("Source output pin '%s' not found on node '%s'"), *SourcePin, *SourceNodeIdStr));
+		}
+	}
+	
+	// Find the target input pin using Builder API
+	InputHandle = Builder->FindNodeInputByName(TargetNodeHandle, FName(*TargetPin), BuilderResult);
+	if (BuilderResult != EMetaSoundBuilderResult::Succeeded)
+	{
+		return FECACommandResult::Error(FString::Printf(TEXT("Target input pin '%s' not found on node '%s'"), *TargetPin, *TargetNodeIdStr));
+	}
+	
+	// Connect using the builder
+	EMetaSoundBuilderResult ConnectResult;
+	Builder->ConnectNodes(OutputHandle, InputHandle, ConnectResult);
+	
+	if (ConnectResult != EMetaSoundBuilderResult::Succeeded)
+	{
+		return FECACommandResult::Error(FString::Printf(TEXT("Failed to connect nodes: %s.%s -> %s.%s"), 
+			*SourceNodeIdStr, *SourcePin, *TargetNodeIdStr, *TargetPin));
+	}
+	
+	MetaSound->MarkPackageDirty();
+	
+	// Synchronize the editor graph to reflect the changes in the UI
+	SynchronizeMetaSoundEditorGraph(MetaSound, Builder);
+	
+	TSharedPtr<FJsonObject> Result = MakeResult();
+	Result->SetStringField(TEXT("source_node_id"), SourceNodeIdStr);
+	Result->SetStringField(TEXT("source_pin"), SourcePin);
+	Result->SetStringField(TEXT("target_node_id"), TargetNodeIdStr);
+	Result->SetStringField(TEXT("target_pin"), TargetPin);
+	
+	return FECACommandResult::Success(Result);
+}
+
+//------------------------------------------------------------------------------
+// DisconnectMetasoundNodes
+//------------------------------------------------------------------------------
+
+FECACommandResult FECACommand_DisconnectMetasoundNodes::Execute(const TSharedPtr<FJsonObject>& Params)
+{
+	FString AssetPath;
+	FString SourceNodeIdStr;
+	FString SourcePin;
+	FString TargetNodeIdStr;
+	FString TargetPin;
+	
+	if (!GetStringParam(Params, TEXT("asset_path"), AssetPath, true))
+	{
+		return FECACommandResult::Error(TEXT("asset_path is required"));
+	}
+	if (!GetStringParam(Params, TEXT("source_node_id"), SourceNodeIdStr, true))
+	{
+		return FECACommandResult::Error(TEXT("source_node_id is required"));
+	}
+	if (!GetStringParam(Params, TEXT("source_pin"), SourcePin, true))
+	{
+		return FECACommandResult::Error(TEXT("source_pin is required"));
+	}
+	if (!GetStringParam(Params, TEXT("target_node_id"), TargetNodeIdStr, true))
+	{
+		return FECACommandResult::Error(TEXT("target_node_id is required"));
+	}
+	if (!GetStringParam(Params, TEXT("target_pin"), TargetPin, true))
+	{
+		return FECACommandResult::Error(TEXT("target_pin is required"));
+	}
+	
+	FGuid SourceNodeId, TargetNodeId;
+	if (!ParseNodeId(SourceNodeIdStr, SourceNodeId))
+	{
+		return FECACommandResult::Error(FString::Printf(TEXT("Invalid source_node_id format: %s"), *SourceNodeIdStr));
+	}
+	if (!ParseNodeId(TargetNodeIdStr, TargetNodeId))
+	{
+		return FECACommandResult::Error(FString::Printf(TEXT("Invalid target_node_id format: %s"), *TargetNodeIdStr));
+	}
+	
+	UMetaSoundSource* MetaSound = LoadMetaSoundSource(AssetPath);
+	if (!MetaSound)
+	{
+		return FECACommandResult::Error(FString::Printf(TEXT("MetaSound not found: %s"), *AssetPath));
+	}
+	
+	FScopedTransaction Transaction(FText::FromString(TEXT("Disconnect MetaSound Nodes")));
+	MetaSound->Modify();
+	
+	// Get the builder using the Editor Subsystem (proper high-level API)
+	EMetaSoundBuilderResult BuilderResult;
+	UMetaSoundBuilderBase* Builder = GetMetaSoundBuilder(MetaSound, BuilderResult);
+	if (!Builder || BuilderResult != EMetaSoundBuilderResult::Succeeded)
+	{
+		return FECACommandResult::Error(TEXT("Failed to get MetaSound builder. Make sure the editor is running."));
+	}
+	
+	// Use the Builder API to find node outputs and inputs by name
+	FMetaSoundNodeHandle SourceNodeHandle(SourceNodeId);
+	FMetaSoundNodeHandle TargetNodeHandle(TargetNodeId);
+	
+	// Find the source output pin using Builder API
+	FMetaSoundBuilderNodeOutputHandle OutputHandle = Builder->FindNodeOutputByName(SourceNodeHandle, FName(*SourcePin), BuilderResult);
+	if (BuilderResult != EMetaSoundBuilderResult::Succeeded)
+	{
+		return FECACommandResult::Error(FString::Printf(TEXT("Source output pin '%s' not found on node '%s'"), *SourcePin, *SourceNodeIdStr));
+	}
+	
+	// Find the target input pin using Builder API
+	FMetaSoundBuilderNodeInputHandle InputHandle = Builder->FindNodeInputByName(TargetNodeHandle, FName(*TargetPin), BuilderResult);
+	if (BuilderResult != EMetaSoundBuilderResult::Succeeded)
+	{
+		return FECACommandResult::Error(FString::Printf(TEXT("Target input pin '%s' not found on node '%s'"), *TargetPin, *TargetNodeIdStr));
+	}
+	
+	// Check if the connection exists
+	if (!Builder->NodesAreConnected(OutputHandle, InputHandle))
+	{
+		return FECACommandResult::Error(TEXT("Connection not found"));
+	}
+	
+	// Disconnect using the builder
+	EMetaSoundBuilderResult DisconnectResult;
+	Builder->DisconnectNodes(OutputHandle, InputHandle, DisconnectResult);
+	
+	if (DisconnectResult != EMetaSoundBuilderResult::Succeeded)
+	{
+		return FECACommandResult::Error(TEXT("Failed to disconnect nodes"));
+	}
+	
+	MetaSound->MarkPackageDirty();
+	
+	// Synchronize the editor graph to reflect the changes in the UI
+	SynchronizeMetaSoundEditorGraph(MetaSound, Builder);
+	
+	TSharedPtr<FJsonObject> Result = MakeResult();
+	Result->SetBoolField(TEXT("disconnected"), true);
+	
+	return FECACommandResult::Success(Result);
+}
+
+//------------------------------------------------------------------------------
+// SetMetasoundNodeInput
+//------------------------------------------------------------------------------
+
+FECACommandResult FECACommand_SetMetasoundNodeInput::Execute(const TSharedPtr<FJsonObject>& Params)
+{
+	FString AssetPath;
+	FString NodeIdStr;
+	FString InputName;
+	
+	if (!GetStringParam(Params, TEXT("asset_path"), AssetPath, true))
+	{
+		return FECACommandResult::Error(TEXT("asset_path is required"));
+	}
+	if (!GetStringParam(Params, TEXT("node_id"), NodeIdStr, true))
+	{
+		return FECACommandResult::Error(TEXT("node_id is required"));
+	}
+	if (!GetStringParam(Params, TEXT("input_name"), InputName, true))
+	{
+		return FECACommandResult::Error(TEXT("input_name is required"));
+	}
+	
+	// Get the value
+	TSharedPtr<FJsonValue> ValueJson = Params->TryGetField(TEXT("value"));
+	if (!ValueJson.IsValid())
+	{
+		return FECACommandResult::Error(TEXT("value is required"));
+	}
+	
+	FGuid NodeId;
+	if (!ParseNodeId(NodeIdStr, NodeId))
+	{
+		return FECACommandResult::Error(FString::Printf(TEXT("Invalid node_id format: %s"), *NodeIdStr));
+	}
+	
+	UMetaSoundSource* MetaSound = LoadMetaSoundSource(AssetPath);
+	if (!MetaSound)
+	{
+		return FECACommandResult::Error(FString::Printf(TEXT("MetaSound not found: %s"), *AssetPath));
+	}
+	
+	FScopedTransaction Transaction(FText::FromString(TEXT("Set MetaSound Node Input")));
+	MetaSound->Modify();
+	
+	// Get the builder using the Editor Subsystem (proper high-level API)
+	EMetaSoundBuilderResult BuilderResult;
+	UMetaSoundBuilderBase* Builder = GetMetaSoundBuilder(MetaSound, BuilderResult);
+	if (!Builder || BuilderResult != EMetaSoundBuilderResult::Succeeded)
+	{
+		return FECACommandResult::Error(TEXT("Failed to get MetaSound builder. Make sure the editor is running."));
+	}
+	
+	// Use the Builder API to find the node input by name
+	FMetaSoundNodeHandle NodeHandle(NodeId);
+	FMetaSoundBuilderNodeInputHandle InputHandle = Builder->FindNodeInputByName(NodeHandle, FName(*InputName), BuilderResult);
+	if (BuilderResult != EMetaSoundBuilderResult::Succeeded)
+	{
+		return FECACommandResult::Error(FString::Printf(TEXT("Input '%s' not found on node '%s'"), *InputName, *NodeIdStr));
+	}
+	
+	// Get the input's data type to create the correct literal type
+	FName InputDataName;
+	FName InputDataType;
+	Builder->GetNodeInputData(InputHandle, InputDataName, InputDataType, BuilderResult);
+	if (BuilderResult != EMetaSoundBuilderResult::Succeeded)
+	{
+		return FECACommandResult::Error(FString::Printf(TEXT("Failed to get data type for input '%s'"), *InputName));
+	}
+	
+	// Build the literal value based on the input's actual data type
+	FMetasoundFrontendLiteral NewLiteral;
+	FString ValueStr;
+	FString DataTypeStr = InputDataType.ToString();
+	
+	// Check if this is an array type
+	bool bIsArrayType = DataTypeStr.Contains(TEXT(":Array")) || DataTypeStr.EndsWith(TEXT("Array"));
+	
+	// Handle array values
+	if (ValueJson->Type == EJson::Array)
+	{
+		const TArray<TSharedPtr<FJsonValue>>& JsonArray = ValueJson->AsArray();
+		
+		if (DataTypeStr.Contains(TEXT("Float")) || DataTypeStr.Contains(TEXT("float")))
+		{
+			TArray<float> FloatArray;
+			for (const TSharedPtr<FJsonValue>& Element : JsonArray)
+			{
+				if (Element->Type == EJson::Number)
+				{
+					FloatArray.Add(static_cast<float>(Element->AsNumber()));
+				}
+				else if (Element->Type == EJson::String)
+				{
+					FloatArray.Add(FCString::Atof(*Element->AsString()));
+				}
+			}
+			NewLiteral.Set(FloatArray);
+			ValueStr = FString::Printf(TEXT("[%d floats]"), FloatArray.Num());
+		}
+		else if (DataTypeStr.Contains(TEXT("Int32")) || DataTypeStr.Contains(TEXT("int32")))
+		{
+			TArray<int32> IntArray;
+			for (const TSharedPtr<FJsonValue>& Element : JsonArray)
+			{
+				if (Element->Type == EJson::Number)
+				{
+					IntArray.Add(static_cast<int32>(Element->AsNumber()));
+				}
+				else if (Element->Type == EJson::String)
+				{
+					IntArray.Add(FCString::Atoi(*Element->AsString()));
+				}
+			}
+			NewLiteral.Set(IntArray);
+			ValueStr = FString::Printf(TEXT("[%d ints]"), IntArray.Num());
+		}
+		else if (DataTypeStr.Contains(TEXT("Bool")) || DataTypeStr.Contains(TEXT("bool")))
+		{
+			TArray<bool> BoolArray;
+			for (const TSharedPtr<FJsonValue>& Element : JsonArray)
+			{
+				if (Element->Type == EJson::Boolean)
+				{
+					BoolArray.Add(Element->AsBool());
+				}
+				else if (Element->Type == EJson::String)
+				{
+					FString Str = Element->AsString();
+					BoolArray.Add(Str.ToBool() || Str == TEXT("1") || Str.Equals(TEXT("true"), ESearchCase::IgnoreCase));
+				}
+				else if (Element->Type == EJson::Number)
+				{
+					BoolArray.Add(Element->AsNumber() != 0.0);
+				}
+			}
+			NewLiteral.Set(BoolArray);
+			ValueStr = FString::Printf(TEXT("[%d bools]"), BoolArray.Num());
+		}
+		else if (DataTypeStr.Contains(TEXT("String")))
+		{
+			TArray<FString> StringArray;
+			for (const TSharedPtr<FJsonValue>& Element : JsonArray)
+			{
+				StringArray.Add(Element->AsString());
+			}
+			NewLiteral.Set(StringArray);
+			ValueStr = FString::Printf(TEXT("[%d strings]"), StringArray.Num());
+		}
+		else
+		{
+			// Default to float array for unknown array types
+			TArray<float> FloatArray;
+			for (const TSharedPtr<FJsonValue>& Element : JsonArray)
+			{
+				if (Element->Type == EJson::Number)
+				{
+					FloatArray.Add(static_cast<float>(Element->AsNumber()));
+				}
+			}
+			NewLiteral.Set(FloatArray);
+			ValueStr = FString::Printf(TEXT("[%d values]"), FloatArray.Num());
+			UE_LOG(LogTemp, Warning, TEXT("ECABridge: Unknown array data type '%s' for input '%s', using float array"), *DataTypeStr, *InputName);
+		}
+	}
+	// Handle scalar values
+	else
+	{
+		// Get the string representation of the value for any type
+		if (ValueJson->Type == EJson::Number)
+		{
+			ValueStr = FString::SanitizeFloat(ValueJson->AsNumber());
+		}
+		else if (ValueJson->Type == EJson::Boolean)
+		{
+			ValueStr = ValueJson->AsBool() ? TEXT("true") : TEXT("false");
+		}
+		else if (ValueJson->Type == EJson::String)
+		{
+			ValueStr = ValueJson->AsString();
+		}
+		else
+		{
+			return FECACommandResult::Error(TEXT("Unsupported value type. Use a number, boolean, string, or array."));
+		}
+		
+		// Create the appropriate literal based on the input's data type
+		if (DataTypeStr == TEXT("Float") || DataTypeStr == TEXT("float"))
+		{
+			float FloatValue = FCString::Atof(*ValueStr);
+			NewLiteral.SetFromLiteral(Metasound::FLiteral(FloatValue));
+		}
+		else if (DataTypeStr == TEXT("Int32") || DataTypeStr == TEXT("int32"))
+		{
+			int32 IntValue = FCString::Atoi(*ValueStr);
+			NewLiteral.SetFromLiteral(Metasound::FLiteral(IntValue));
+		}
+		else if (DataTypeStr == TEXT("Bool") || DataTypeStr == TEXT("bool"))
+		{
+			bool BoolValue = ValueStr.ToBool() || ValueStr == TEXT("1") || ValueStr.Equals(TEXT("true"), ESearchCase::IgnoreCase);
+			NewLiteral.SetFromLiteral(Metasound::FLiteral(BoolValue));
+		}
+		else if (DataTypeStr == TEXT("String") || DataTypeStr == TEXT("FString"))
+		{
+			NewLiteral.SetFromLiteral(Metasound::FLiteral(ValueStr));
+		}
+		else if (DataTypeStr == TEXT("Time") || DataTypeStr == TEXT("Audio:Time"))
+		{
+			// Time is stored as a float representing seconds
+			float TimeValue = FCString::Atof(*ValueStr);
+			NewLiteral.SetFromLiteral(Metasound::FLiteral(TimeValue));
+		}
+		else
+		{
+			// For unknown types, try to use the JSON type as a fallback
+			if (ValueJson->Type == EJson::Number)
+			{
+				NewLiteral.SetFromLiteral(Metasound::FLiteral(static_cast<float>(ValueJson->AsNumber())));
+			}
+			else if (ValueJson->Type == EJson::Boolean)
+			{
+				NewLiteral.SetFromLiteral(Metasound::FLiteral(ValueJson->AsBool()));
+			}
+			else
+			{
+				NewLiteral.SetFromLiteral(Metasound::FLiteral(ValueStr));
+			}
+			UE_LOG(LogTemp, Warning, TEXT("ECABridge: Unknown data type '%s' for input '%s', using fallback literal creation"), *DataTypeStr, *InputName);
+		}
+	}
+	
+	// Use the Builder API to set the node input default value
+	Builder->SetNodeInputDefault(InputHandle, NewLiteral, BuilderResult);
+	if (BuilderResult != EMetaSoundBuilderResult::Succeeded)
+	{
+		return FECACommandResult::Error(FString::Printf(TEXT("Failed to set value for input '%s'"), *InputName));
+	}
+	
+	MetaSound->MarkPackageDirty();
+	
+	// Synchronize the editor graph to reflect the changes in the UI
+	SynchronizeMetaSoundEditorGraph(MetaSound);
+	
+	TSharedPtr<FJsonObject> Result = MakeResult();
+	Result->SetStringField(TEXT("node_id"), NodeIdStr);
+	Result->SetStringField(TEXT("input_name"), InputName);
+	Result->SetStringField(TEXT("value"), ValueStr);
+	
+	return FECACommandResult::Success(Result);
+}
+
+//------------------------------------------------------------------------------
+// AddMetasoundInput
+//------------------------------------------------------------------------------
+
+FECACommandResult FECACommand_AddMetasoundInput::Execute(const TSharedPtr<FJsonObject>& Params)
+{
+	FString AssetPath;
+	FString InputName;
+	FString DataType;
+	
+	if (!GetStringParam(Params, TEXT("asset_path"), AssetPath, true))
+	{
+		return FECACommandResult::Error(TEXT("asset_path is required"));
+	}
+	if (!GetStringParam(Params, TEXT("input_name"), InputName, true))
+	{
+		return FECACommandResult::Error(TEXT("input_name is required"));
+	}
+	if (!GetStringParam(Params, TEXT("data_type"), DataType, true))
+	{
+		return FECACommandResult::Error(TEXT("data_type is required"));
+	}
+	
+	UMetaSoundSource* MetaSound = LoadMetaSoundSource(AssetPath);
+	if (!MetaSound)
+	{
+		return FECACommandResult::Error(FString::Printf(TEXT("MetaSound not found: %s"), *AssetPath));
+	}
+	
+	FScopedTransaction Transaction(FText::FromString(TEXT("Add MetaSound Input")));
+	MetaSound->Modify();
+	
+	// Get the builder using the Editor Subsystem (proper high-level API)
+	EMetaSoundBuilderResult BuilderResult;
+	UMetaSoundBuilderBase* Builder = GetMetaSoundBuilder(MetaSound, BuilderResult);
+	if (!Builder || BuilderResult != EMetaSoundBuilderResult::Succeeded)
+	{
+		return FECACommandResult::Error(TEXT("Failed to get MetaSound builder. Make sure the editor is running."));
+	}
+	
+	// Parse optional default value
+	FMetasoundFrontendLiteral DefaultLiteral;
+	TSharedPtr<FJsonValue> DefaultValueJson = Params->TryGetField(TEXT("default_value"));
+	if (DefaultValueJson.IsValid())
+	{
+		// Handle array values
+		if (DefaultValueJson->Type == EJson::Array)
+		{
+			const TArray<TSharedPtr<FJsonValue>>& JsonArray = DefaultValueJson->AsArray();
+			
+			if (DataType.Contains(TEXT("Float")) || DataType.Contains(TEXT("float")))
+			{
+				TArray<float> FloatArray;
+				for (const TSharedPtr<FJsonValue>& Element : JsonArray)
+				{
+					if (Element->Type == EJson::Number)
+					{
+						FloatArray.Add(static_cast<float>(Element->AsNumber()));
+					}
+					else if (Element->Type == EJson::String)
+					{
+						FloatArray.Add(FCString::Atof(*Element->AsString()));
+					}
+				}
+				DefaultLiteral.Set(FloatArray);
+			}
+			else if (DataType.Contains(TEXT("Int32")) || DataType.Contains(TEXT("int32")))
+			{
+				TArray<int32> IntArray;
+				for (const TSharedPtr<FJsonValue>& Element : JsonArray)
+				{
+					if (Element->Type == EJson::Number)
+					{
+						IntArray.Add(static_cast<int32>(Element->AsNumber()));
+					}
+					else if (Element->Type == EJson::String)
+					{
+						IntArray.Add(FCString::Atoi(*Element->AsString()));
+					}
+				}
+				DefaultLiteral.Set(IntArray);
+			}
+			else if (DataType.Contains(TEXT("Bool")) || DataType.Contains(TEXT("bool")))
+			{
+				TArray<bool> BoolArray;
+				for (const TSharedPtr<FJsonValue>& Element : JsonArray)
+				{
+					if (Element->Type == EJson::Boolean)
+					{
+						BoolArray.Add(Element->AsBool());
+					}
+					else if (Element->Type == EJson::Number)
+					{
+						BoolArray.Add(Element->AsNumber() != 0.0);
+					}
+				}
+				DefaultLiteral.Set(BoolArray);
+			}
+			else
+			{
+				// Default to float array
+				TArray<float> FloatArray;
+				for (const TSharedPtr<FJsonValue>& Element : JsonArray)
+				{
+					if (Element->Type == EJson::Number)
+					{
+						FloatArray.Add(static_cast<float>(Element->AsNumber()));
+					}
+				}
+				DefaultLiteral.Set(FloatArray);
+			}
+		}
+		// Handle scalar values
+		else if (DefaultValueJson->Type == EJson::Number)
+		{
+			if (DataType.Contains(TEXT("Int32")) || DataType.Contains(TEXT("int32")))
+			{
+				DefaultLiteral.Set(static_cast<int32>(DefaultValueJson->AsNumber()));
+			}
+			else
+			{
+				DefaultLiteral.Set(static_cast<float>(DefaultValueJson->AsNumber()));
+			}
+		}
+		else if (DefaultValueJson->Type == EJson::Boolean)
+		{
+			DefaultLiteral.Set(DefaultValueJson->AsBool());
+		}
+		else if (DefaultValueJson->Type == EJson::String)
+		{
+			FString ValueStr = DefaultValueJson->AsString();
+			if (DataType.Contains(TEXT("Float")) || DataType.Contains(TEXT("float")) || DataType.Contains(TEXT("Time")))
+			{
+				DefaultLiteral.Set(FCString::Atof(*ValueStr));
+			}
+			else if (DataType.Contains(TEXT("Int32")) || DataType.Contains(TEXT("int32")))
+			{
+				DefaultLiteral.Set(FCString::Atoi(*ValueStr));
+			}
+			else if (DataType.Contains(TEXT("Bool")) || DataType.Contains(TEXT("bool")))
+			{
+				DefaultLiteral.Set(ValueStr.ToBool() || ValueStr == TEXT("1"));
+			}
+			else
+			{
+				DefaultLiteral.Set(ValueStr);
+			}
+		}
+	}
+	
+	// Use Builder API to add the graph input - this creates both the interface entry and the input node
+	FMetaSoundBuilderNodeOutputHandle OutputHandle = Builder->AddGraphInputNode(
+		FName(*InputName), 
+		FName(*DataType), 
+		DefaultLiteral, 
+		BuilderResult
+	);
+	
+	if (BuilderResult != EMetaSoundBuilderResult::Succeeded || !OutputHandle.IsSet())
+	{
+		return FECACommandResult::Error(FString::Printf(TEXT("Failed to add graph input '%s' with type '%s'"), *InputName, *DataType));
+	}
+	
+	MetaSound->MarkPackageDirty();
+	
+	// Synchronize the editor graph to reflect the changes in the UI
+	SynchronizeMetaSoundEditorGraph(MetaSound, Builder);
+	
+	TSharedPtr<FJsonObject> Result = MakeResult();
+	Result->SetStringField(TEXT("input_name"), InputName);
+	Result->SetStringField(TEXT("data_type"), DataType);
+	Result->SetStringField(TEXT("node_id"), GetNodeIdString(OutputHandle.NodeID));
+	Result->SetStringField(TEXT("vertex_id"), GetNodeIdString(OutputHandle.VertexID));
+	
+	return FECACommandResult::Success(Result);
+}
+
+//------------------------------------------------------------------------------
+// RemoveMetasoundInput
+//------------------------------------------------------------------------------
+
+FECACommandResult FECACommand_RemoveMetasoundInput::Execute(const TSharedPtr<FJsonObject>& Params)
+{
+	FString AssetPath;
+	FString InputName;
+	
+	if (!GetStringParam(Params, TEXT("asset_path"), AssetPath, true))
+	{
+		return FECACommandResult::Error(TEXT("asset_path is required"));
+	}
+	if (!GetStringParam(Params, TEXT("input_name"), InputName, true))
+	{
+		return FECACommandResult::Error(TEXT("input_name is required"));
+	}
+	
+	UMetaSoundSource* MetaSound = LoadMetaSoundSource(AssetPath);
+	if (!MetaSound)
+	{
+		return FECACommandResult::Error(FString::Printf(TEXT("MetaSound not found: %s"), *AssetPath));
+	}
+	
+	FScopedTransaction Transaction(FText::FromString(TEXT("Remove MetaSound Input")));
+	MetaSound->Modify();
+	
+	FMetasoundFrontendDocument& Document = MetaSound->GetDocumentChecked();
+	FMetasoundFrontendGraphClass& GraphClass = Document.RootGraph;
+	
+	int32 RemovedCount = GraphClass.GetDefaultInterface().Inputs.RemoveAll(
+		[&InputName](const FMetasoundFrontendClassInput& Input)
+		{
+			return Input.Name.ToString().Equals(InputName, ESearchCase::IgnoreCase);
+		});
+	
+	if (RemovedCount == 0)
+	{
+		return FECACommandResult::Error(FString::Printf(TEXT("Input '%s' not found"), *InputName));
+	}
+	
+	MetaSound->MarkPackageDirty();
+	
+	// Synchronize the editor graph to reflect the changes in the UI
+	SynchronizeMetaSoundEditorGraph(MetaSound);
+	
+	TSharedPtr<FJsonObject> Result = MakeResult();
+	Result->SetStringField(TEXT("removed_input"), InputName);
+	
+	return FECACommandResult::Success(Result);
+}
+
+//------------------------------------------------------------------------------
+// GetMetasoundNodeTypes
+//------------------------------------------------------------------------------
+
+FECACommandResult FECACommand_GetMetasoundNodeTypes::Execute(const TSharedPtr<FJsonObject>& Params)
+{
+	using namespace Metasound::Frontend;
+	
+	FString CategoryFilter;
+	FString NameFilter;
+	GetStringParam(Params, TEXT("category_filter"), CategoryFilter, false);
+	GetStringParam(Params, TEXT("name_filter"), NameFilter, false);
+	
+	TArray<TSharedPtr<FJsonValue>> NodeTypesArray;
+	TSet<FString> Categories;
+	
+#if WITH_EDITORONLY_DATA
+	ISearchEngine& SearchEngine = ISearchEngine::Get();
+	TArray<FMetaSoundClassInfo> AllClasses = SearchEngine.FindAllClasses(ISearchEngine::EResultVersion::Highest, false);
+	
+	for (const FMetaSoundClassInfo& ClassInfo : AllClasses)
+	{
+		FString ClassName = ClassInfo.ClassName.Name.ToString();
+		FString Namespace = ClassInfo.ClassName.Namespace.ToString();
+		FString DisplayName = ClassInfo.InterfaceInfo.SearchInfo.ClassDisplayName.ToString();
+		FString Description = ClassInfo.InterfaceInfo.SearchInfo.ClassDescription.ToString();
+		
+		// Build category from Hierarchy (e.g., "Generators", "Filters", "Envelopes")
+		// Hierarchy contains the menu path elements used in the MetaSound editor
+		FString Category;
+		TArray<FString> HierarchyParts;
+		for (const FText& HierarchyPart : ClassInfo.InterfaceInfo.SearchInfo.Hierarchy)
+		{
+			FString PartStr = HierarchyPart.ToString();
+			if (!PartStr.IsEmpty())
+			{
+				HierarchyParts.Add(PartStr);
+			}
+		}
+		
+		if (HierarchyParts.Num() > 0)
+		{
+			Category = FString::Join(HierarchyParts, TEXT("|"));
+		}
+		else
+		{
+			// Fallback to namespace if no hierarchy
+			Category = Namespace;
+		}
+		
+		// Category filter: check against hierarchy parts or full category path (case-insensitive)
+		if (!CategoryFilter.IsEmpty())
+		{
+			bool bMatchesCategory = Category.Contains(CategoryFilter, ESearchCase::IgnoreCase);
+			// Also check individual hierarchy parts for partial match
+			if (!bMatchesCategory)
+			{
+				for (const FString& Part : HierarchyParts)
+				{
+					if (Part.Contains(CategoryFilter, ESearchCase::IgnoreCase))
+					{
+						bMatchesCategory = true;
+						break;
+					}
+				}
+			}
+			// Also check namespace as fallback
+			if (!bMatchesCategory && Namespace.Contains(CategoryFilter, ESearchCase::IgnoreCase))
+			{
+				bMatchesCategory = true;
+			}
+			if (!bMatchesCategory)
+			{
+				continue;
+			}
+		}
+		
+		// Name filter (case-insensitive)
+		if (!NameFilter.IsEmpty() && 
+			!DisplayName.Contains(NameFilter, ESearchCase::IgnoreCase) && 
+			!ClassName.Contains(NameFilter, ESearchCase::IgnoreCase))
+		{
+			continue;
+		}
+		
+		// Add primary category (first hierarchy element or namespace)
+		if (HierarchyParts.Num() > 0)
+		{
+			Categories.Add(HierarchyParts[0]);
+		}
+		else if (!Namespace.IsEmpty())
+		{
+			Categories.Add(Namespace);
+		}
+		
+		TSharedPtr<FJsonObject> NodeTypeObj = MakeShared<FJsonObject>();
+		NodeTypeObj->SetStringField(TEXT("name"), ClassName);
+		NodeTypeObj->SetStringField(TEXT("display_name"), DisplayName.IsEmpty() ? ClassName : DisplayName);
+		NodeTypeObj->SetStringField(TEXT("namespace"), Namespace);
+		NodeTypeObj->SetStringField(TEXT("category"), Category);
+		NodeTypeObj->SetStringField(TEXT("description"), Description);
+		NodeTypeObj->SetNumberField(TEXT("version_major"), ClassInfo.Version.Major);
+		NodeTypeObj->SetNumberField(TEXT("version_minor"), ClassInfo.Version.Minor);
+		
+		// Include hierarchy array for more detailed categorization
+		TArray<TSharedPtr<FJsonValue>> HierarchyArray;
+		for (const FString& Part : HierarchyParts)
+		{
+			HierarchyArray.Add(MakeShared<FJsonValueString>(Part));
+		}
+		NodeTypeObj->SetArrayField(TEXT("hierarchy"), HierarchyArray);
+		
+		NodeTypesArray.Add(MakeShared<FJsonValueObject>(NodeTypeObj));
+	}
+#endif
+	
+	TArray<TSharedPtr<FJsonValue>> CategoriesArray;
+	for (const FString& Cat : Categories)
+	{
+		CategoriesArray.Add(MakeShared<FJsonValueString>(Cat));
+	}
+	
+	TSharedPtr<FJsonObject> Result = MakeResult();
+	Result->SetArrayField(TEXT("node_types"), NodeTypesArray);
+	Result->SetNumberField(TEXT("count"), NodeTypesArray.Num());
+	Result->SetArrayField(TEXT("categories"), CategoriesArray);
+	
+	return FECACommandResult::Success(Result);
+}
+
+//------------------------------------------------------------------------------
+// PreviewMetasound
+//------------------------------------------------------------------------------
+
+FECACommandResult FECACommand_PreviewMetasound::Execute(const TSharedPtr<FJsonObject>& Params)
+{
+	FString AssetPath;
+	FString Action = TEXT("play");
+	
+	if (!GetStringParam(Params, TEXT("asset_path"), AssetPath, true))
+	{
+		return FECACommandResult::Error(TEXT("asset_path is required"));
+	}
+	GetStringParam(Params, TEXT("action"), Action, false);
+	
+	UMetaSoundSource* MetaSound = LoadMetaSoundSource(AssetPath);
+	if (!MetaSound)
+	{
+		return FECACommandResult::Error(FString::Printf(TEXT("MetaSound not found: %s"), *AssetPath));
+	}
+	
+	if (GEditor)
+	{
+		if (Action.Equals(TEXT("play"), ESearchCase::IgnoreCase))
+		{
+			GEditor->PlayPreviewSound(MetaSound);
+		}
+		else if (Action.Equals(TEXT("stop"), ESearchCase::IgnoreCase))
+		{
+			GEditor->ResetPreviewAudioComponent();
+		}
+		else
+		{
+			return FECACommandResult::Error(FString::Printf(TEXT("Unknown action: %s. Use 'play' or 'stop'"), *Action));
+		}
+	}
+	else
+	{
+		return FECACommandResult::Error(TEXT("Editor not available"));
+	}
+	
+	TSharedPtr<FJsonObject> Result = MakeResult();
+	Result->SetStringField(TEXT("asset_path"), AssetPath);
+	Result->SetStringField(TEXT("action"), Action);
+	
+	return FECACommandResult::Success(Result);
+}
+
+//------------------------------------------------------------------------------
+// GetMetasoundInterface
+//------------------------------------------------------------------------------
+
+FECACommandResult FECACommand_GetMetasoundInterface::Execute(const TSharedPtr<FJsonObject>& Params)
+{
+	FString AssetPath;
+	
+	if (!GetStringParam(Params, TEXT("asset_path"), AssetPath, true))
+	{
+		return FECACommandResult::Error(TEXT("asset_path is required"));
+	}
+	
+	UMetaSoundSource* MetaSound = LoadMetaSoundSource(AssetPath);
+	if (!MetaSound)
+	{
+		return FECACommandResult::Error(FString::Printf(TEXT("MetaSound not found: %s"), *AssetPath));
+	}
+	
+	const FMetasoundFrontendDocument& Document = MetaSound->GetConstDocumentChecked();
+	const FMetasoundFrontendGraphClass& GraphClass = Document.RootGraph;
+	const FMetasoundFrontendClassInterface& Interface = GraphClass.GetDefaultInterface();
+	
+	// Get inputs
+	TArray<TSharedPtr<FJsonValue>> InputsArray;
+	for (const FMetasoundFrontendClassInput& Input : Interface.Inputs)
+	{
+		TSharedPtr<FJsonObject> InputObj = MakeShared<FJsonObject>();
+		InputObj->SetStringField(TEXT("name"), Input.Name.ToString());
+		InputObj->SetStringField(TEXT("type"), Input.TypeName.ToString());
+		InputObj->SetStringField(TEXT("vertex_id"), GetNodeIdString(Input.VertexID));
+		InputObj->SetStringField(TEXT("node_id"), GetNodeIdString(Input.NodeID));
+		InputsArray.Add(MakeShared<FJsonValueObject>(InputObj));
+	}
+	
+	// Get outputs
+	TArray<TSharedPtr<FJsonValue>> OutputsArray;
+	for (const FMetasoundFrontendClassOutput& Output : Interface.Outputs)
+	{
+		TSharedPtr<FJsonObject> OutputObj = MakeShared<FJsonObject>();
+		OutputObj->SetStringField(TEXT("name"), Output.Name.ToString());
+		OutputObj->SetStringField(TEXT("type"), Output.TypeName.ToString());
+		OutputObj->SetStringField(TEXT("vertex_id"), GetNodeIdString(Output.VertexID));
+		OutputObj->SetStringField(TEXT("node_id"), GetNodeIdString(Output.NodeID));
+		OutputsArray.Add(MakeShared<FJsonValueObject>(OutputObj));
+	}
+	
+	TSharedPtr<FJsonObject> Result = MakeResult();
+	Result->SetArrayField(TEXT("inputs"), InputsArray);
+	Result->SetArrayField(TEXT("outputs"), OutputsArray);
+	Result->SetNumberField(TEXT("input_count"), InputsArray.Num());
+	Result->SetNumberField(TEXT("output_count"), OutputsArray.Num());
+	
+	return FECACommandResult::Success(Result);
+}
+
+//------------------------------------------------------------------------------
+// SpawnMetasoundPlayer
+//------------------------------------------------------------------------------
+
+FECACommandResult FECACommand_SpawnMetasoundPlayer::Execute(const TSharedPtr<FJsonObject>& Params)
+{
+	FString AssetPath;
+	FVector Location = FVector::ZeroVector;
+	FString Name;
+	bool bAutoPlay = true;
+	FString Folder;
+	
+	if (!GetStringParam(Params, TEXT("asset_path"), AssetPath, true))
+	{
+		return FECACommandResult::Error(TEXT("asset_path is required"));
+	}
+	GetVectorParam(Params, TEXT("location"), Location, false);
+	GetStringParam(Params, TEXT("name"), Name, false);
+	GetBoolParam(Params, TEXT("auto_play"), bAutoPlay, false);
+	GetStringParam(Params, TEXT("folder"), Folder, false);
+	
+	UMetaSoundSource* MetaSound = LoadMetaSoundSource(AssetPath);
+	if (!MetaSound)
+	{
+		return FECACommandResult::Error(FString::Printf(TEXT("MetaSound not found: %s"), *AssetPath));
+	}
+	
+	UWorld* World = GetEditorWorld();
+	if (!World)
+	{
+		return FECACommandResult::Error(TEXT("No editor world available"));
+	}
+	
+	FActorSpawnParameters SpawnParams;
+	SpawnParams.SpawnCollisionHandlingOverride = ESpawnActorCollisionHandlingMethod::AlwaysSpawn;
+	
+	AActor* NewActor = World->SpawnActor<AActor>(AActor::StaticClass(), Location, FRotator::ZeroRotator, SpawnParams);
+	if (!NewActor)
+	{
+		return FECACommandResult::Error(TEXT("Failed to spawn actor"));
+	}
+	
+	if (!Name.IsEmpty())
+	{
+		NewActor->SetActorLabel(Name);
+	}
+	else
+	{
+		NewActor->SetActorLabel(FString::Printf(TEXT("MetaSound_%s"), *FPackageName::GetShortName(AssetPath)));
+	}
+	
+	UAudioComponent* AudioComponent = NewObject<UAudioComponent>(NewActor, NAME_None, RF_Transactional);
+	AudioComponent->SetSound(MetaSound);
+	AudioComponent->bAutoActivate = bAutoPlay;
+	AudioComponent->RegisterComponent();
+	NewActor->AddInstanceComponent(AudioComponent);
+	NewActor->SetRootComponent(AudioComponent);
+	
+	if (!Folder.IsEmpty())
+	{
+		NewActor->SetFolderPath(FName(*Folder));
+	}
+	
+	if (bAutoPlay)
+	{
+		AudioComponent->Play();
+	}
+	
+	TSharedPtr<FJsonObject> Result = MakeResult();
+	Result->SetStringField(TEXT("actor_name"), NewActor->GetActorLabel());
+	Result->SetStringField(TEXT("asset_path"), AssetPath);
+	Result->SetObjectField(TEXT("location"), VectorToJson(Location));
+	Result->SetBoolField(TEXT("auto_play"), bAutoPlay);
+	
+	return FECACommandResult::Success(Result);
+}
+
+
+//------------------------------------------------------------------------------
+// AutoLayoutMetasoundGraph - Automatically arrange MetaSound nodes
+//------------------------------------------------------------------------------
+
+FECACommandResult FECACommand_AutoLayoutMetasoundGraph::Execute(const TSharedPtr<FJsonObject>& Params)
+{
+	using namespace Metasound::Frontend;
+	
+	FString AssetPath;
+	if (!GetStringParam(Params, TEXT("asset_path"), AssetPath, true))
+	{
+		return FECACommandResult::Error(TEXT("asset_path is required"));
+	}
+	
+	UMetaSoundSource* MetaSound = LoadMetaSoundSource(AssetPath);
+	if (!MetaSound)
+	{
+		return FECACommandResult::Error(FString::Printf(TEXT("MetaSound not found: %s"), *AssetPath));
+	}
+	
+	// Get layout parameters
+	FString Strategy = TEXT("horizontal");
+	GetStringParam(Params, TEXT("strategy"), Strategy, false);
+	Strategy = Strategy.ToLower();
+	
+	int32 SpacingX = 300;
+	int32 SpacingY = 100;
+	GetIntParam(Params, TEXT("spacing_x"), SpacingX, false);
+	GetIntParam(Params, TEXT("spacing_y"), SpacingY, false);
+	
+	// Get specific node IDs if provided
+	TSet<FGuid> SpecificNodeIds;
+	const TArray<TSharedPtr<FJsonValue>>* NodeIdsArray;
+	if (GetArrayParam(Params, TEXT("node_ids"), NodeIdsArray, false) && NodeIdsArray)
+	{
+		for (const TSharedPtr<FJsonValue>& IdValue : *NodeIdsArray)
+		{
+			FString IdStr;
+			if (IdValue->TryGetString(IdStr))
+			{
+				FGuid Guid;
+				if (ParseNodeId(IdStr, Guid))
+				{
+					SpecificNodeIds.Add(Guid);
+				}
+			}
+		}
+	}
+	
+	FScopedTransaction Transaction(FText::FromString(TEXT("Auto Layout MetaSound Graph")));
+	MetaSound->Modify();
+	
+	// Get the builder
+	EMetaSoundBuilderResult BuilderResult;
+	UMetaSoundBuilderBase* Builder = GetMetaSoundBuilder(MetaSound, BuilderResult);
+	if (!Builder || BuilderResult != EMetaSoundBuilderResult::Succeeded)
+	{
+		return FECACommandResult::Error(TEXT("Failed to get MetaSound builder. Make sure the editor is running."));
+	}
+	
+	const FMetaSoundFrontendDocumentBuilder& DocBuilder = Builder->GetConstBuilder();
+	const FMetasoundFrontendDocument& Document = DocBuilder.GetConstDocumentChecked();
+	const FMetasoundFrontendGraphClass& GraphClass = Document.RootGraph;
+	
+	// Collect all nodes
+	TArray<FMetasoundFrontendNode> AllNodes;
+	TMap<FGuid, FMetasoundFrontendNode> NodeMap;
+	
+	DocBuilder.IterateNodesByPredicate(
+		[&](const FMetasoundFrontendClass& NodeClass, const FMetasoundFrontendNode& Node, const FGuid& PageID)
+		{
+			// Filter by specific node IDs if provided
+			if (SpecificNodeIds.Num() > 0 && !SpecificNodeIds.Contains(Node.GetID()))
+			{
+				return;
+			}
+			AllNodes.Add(Node);
+			NodeMap.Add(Node.GetID(), Node);
+		},
+		[](const FMetasoundFrontendNode& Node) { return true; },
+		nullptr,
+		true
+	);
+	
+	if (AllNodes.Num() == 0)
+	{
+		return FECACommandResult::Error(TEXT("No nodes to layout"));
+	}
+	
+	// Build connection maps
+	TMap<FGuid, TSet<FGuid>> ConsumersMap; // NodeID -> Set of nodes that consume from it
+	TMap<FGuid, TSet<FGuid>> ProducersMap; // NodeID -> Set of nodes that produce for it
+	
+	// Iterate edges to build connection maps
+	GraphClass.IterateGraphPages([&](const FMetasoundFrontendGraph& Graph)
+	{
+		for (const FMetasoundFrontendEdge& Edge : Graph.Edges)
+		{
+			if (NodeMap.Contains(Edge.FromNodeID) && NodeMap.Contains(Edge.ToNodeID))
+			{
+				ConsumersMap.FindOrAdd(Edge.FromNodeID).Add(Edge.ToNodeID);
+				ProducersMap.FindOrAdd(Edge.ToNodeID).Add(Edge.FromNodeID);
+			}
+		}
+	});
+	
+	// Find input nodes (nodes with no producers / input connections)
+	TArray<FGuid> InputNodeIds;
+	TMap<FGuid, int32> NodeDepths;
+	
+	for (const FMetasoundFrontendNode& Node : AllNodes)
+	{
+		FGuid NodeId = Node.GetID();
+		
+		// Check if this node has no producers (input node)
+		if (!ProducersMap.Contains(NodeId) || ProducersMap[NodeId].Num() == 0)
+		{
+			InputNodeIds.Add(NodeId);
+			NodeDepths.Add(NodeId, 0);
+		}
+	}
+	
+	// If no input nodes found, use all nodes that have outputs
+	if (InputNodeIds.Num() == 0)
+	{
+		for (const FMetasoundFrontendNode& Node : AllNodes)
+		{
+			FGuid NodeId = Node.GetID();
+			if (ConsumersMap.Contains(NodeId) && ConsumersMap[NodeId].Num() > 0)
+			{
+				InputNodeIds.Add(NodeId);
+				NodeDepths.Add(NodeId, 0);
+				break; // Just use the first one as root
+			}
+		}
+	}
+	
+	// If still no roots, just use the first node
+	if (InputNodeIds.Num() == 0 && AllNodes.Num() > 0)
+	{
+		InputNodeIds.Add(AllNodes[0].GetID());
+		NodeDepths.Add(AllNodes[0].GetID(), 0);
+	}
+	
+	// BFS to assign depths
+	TQueue<FGuid> TraversalQueue;
+	TSet<FGuid> Visited;
+	
+	for (const FGuid& InputId : InputNodeIds)
+	{
+		TraversalQueue.Enqueue(InputId);
+		Visited.Add(InputId);
+	}
+	
+	while (!TraversalQueue.IsEmpty())
+	{
+		FGuid CurrentId;
+		TraversalQueue.Dequeue(CurrentId);
+		
+		int32 CurrentDepth = NodeDepths.FindRef(CurrentId);
+		
+		// Find consumers (nodes that take input from this node)
+		if (ConsumersMap.Contains(CurrentId))
+		{
+			for (const FGuid& ConsumerId : ConsumersMap[CurrentId])
+			{
+				if (!Visited.Contains(ConsumerId))
+				{
+					Visited.Add(ConsumerId);
+					NodeDepths.Add(ConsumerId, CurrentDepth + 1);
+					TraversalQueue.Enqueue(ConsumerId);
+				}
+			}
+		}
+	}
+	
+	// Handle unvisited nodes
+	for (const FMetasoundFrontendNode& Node : AllNodes)
+	{
+		if (!NodeDepths.Contains(Node.GetID()))
+		{
+			NodeDepths.Add(Node.GetID(), 0);
+		}
+	}
+	
+	// Group nodes by depth
+	TMap<int32, TArray<FGuid>> NodesByDepth;
+	int32 MaxDepth = 0;
+	
+	for (const auto& Pair : NodeDepths)
+	{
+		NodesByDepth.FindOrAdd(Pair.Value).Add(Pair.Key);
+		MaxDepth = FMath::Max(MaxDepth, Pair.Value);
+	}
+	
+	// Calculate starting position
+	double StartX = 0.0;
+	double StartY = 0.0;
+	
+#if WITH_EDITORONLY_DATA
+	// Use minimum position of existing nodes as anchor
+	bool bFoundPosition = false;
+	for (const FMetasoundFrontendNode& Node : AllNodes)
+	{
+		if (Node.Style.Display.Locations.Num() > 0)
+		{
+			for (const auto& LocationPair : Node.Style.Display.Locations)
+			{
+				if (!bFoundPosition)
+				{
+					StartX = LocationPair.Value.X;
+					StartY = LocationPair.Value.Y;
+					bFoundPosition = true;
+				}
+				else
+				{
+					StartX = FMath::Min(StartX, (double)LocationPair.Value.X);
+					StartY = FMath::Min(StartY, (double)LocationPair.Value.Y);
+				}
+			}
+		}
+	}
+#endif
+	
+	// Position nodes based on strategy
+	int32 NodesPositioned = 0;
+	TMap<FGuid, FVector2D> NewPositions;
+	
+	if (Strategy == TEXT("vertical"))
+	{
+		// Vertical: top-to-bottom following signal flow
+		for (int32 Depth = 0; Depth <= MaxDepth; Depth++)
+		{
+			TArray<FGuid>* NodesAtDepth = NodesByDepth.Find(Depth);
+			if (!NodesAtDepth)
+			{
+				continue;
+			}
+			
+			double Y = StartY + Depth * SpacingY;
+			double X = StartX;
+			
+			for (const FGuid& NodeId : *NodesAtDepth)
+			{
+				NewPositions.Add(NodeId, FVector2D(X, Y));
+				X += SpacingX;
+				NodesPositioned++;
+			}
+		}
+	}
+	else if (Strategy == TEXT("tree"))
+	{
+		// Tree layout - center children under parents
+		TMap<FGuid, int32> SubtreeWidths;
+		
+		for (int32 Depth = MaxDepth; Depth >= 0; Depth--)
+		{
+			TArray<FGuid>* NodesAtDepth = NodesByDepth.Find(Depth);
+			if (!NodesAtDepth)
+			{
+				continue;
+			}
+			
+			for (const FGuid& NodeId : *NodesAtDepth)
+			{
+				int32 Width = 1;
+				
+				if (ConsumersMap.Contains(NodeId))
+				{
+					for (const FGuid& ConsumerId : ConsumersMap[NodeId])
+					{
+						if (SubtreeWidths.Contains(ConsumerId))
+						{
+							Width += SubtreeWidths[ConsumerId];
+						}
+					}
+				}
+				
+				SubtreeWidths.Add(NodeId, FMath::Max(1, Width));
+			}
+		}
+		
+		// Position roots
+		double CurrentY = StartY;
+		for (const FGuid& RootId : InputNodeIds)
+		{
+			int32 Width = SubtreeWidths.FindRef(RootId);
+			NewPositions.Add(RootId, FVector2D(StartX, CurrentY + (Width * SpacingY) / 2.0));
+			CurrentY += Width * SpacingY;
+			NodesPositioned++;
+		}
+		
+		// Position remaining depths
+		for (int32 Depth = 1; Depth <= MaxDepth; Depth++)
+		{
+			TArray<FGuid>* NodesAtDepth = NodesByDepth.Find(Depth);
+			if (!NodesAtDepth)
+			{
+				continue;
+			}
+			
+			double X = StartX + Depth * SpacingX;
+			
+			for (const FGuid& NodeId : *NodesAtDepth)
+			{
+				if (!NewPositions.Contains(NodeId))
+				{
+					NewPositions.Add(NodeId, FVector2D(X, StartY));
+				}
+				else
+				{
+					FVector2D Pos = NewPositions[NodeId];
+					NewPositions[NodeId] = FVector2D(X, Pos.Y);
+				}
+				NodesPositioned++;
+			}
+		}
+	}
+	else if (Strategy == TEXT("compact"))
+	{
+		// Compact grid layout
+		int32 NodesPerRow = FMath::Max(1, FMath::CeilToInt(FMath::Sqrt((float)AllNodes.Num())));
+		int32 Index = 0;
+		
+		for (const FMetasoundFrontendNode& Node : AllNodes)
+		{
+			int32 Row = Index / NodesPerRow;
+			int32 Col = Index % NodesPerRow;
+			
+			NewPositions.Add(Node.GetID(), FVector2D(StartX + Col * SpacingX, StartY + Row * SpacingY));
+			
+			Index++;
+			NodesPositioned++;
+		}
+	}
+	else // horizontal (default)
+	{
+		// Horizontal: left-to-right following signal flow
+		for (int32 Depth = 0; Depth <= MaxDepth; Depth++)
+		{
+			TArray<FGuid>* NodesAtDepth = NodesByDepth.Find(Depth);
+			if (!NodesAtDepth)
+			{
+				continue;
+			}
+			
+			double X = StartX + Depth * SpacingX;
+			double Y = StartY;
+			
+			for (const FGuid& NodeId : *NodesAtDepth)
+			{
+				NewPositions.Add(NodeId, FVector2D(X, Y));
+				Y += SpacingY;
+				NodesPositioned++;
+			}
+		}
+	}
+	
+	// Apply positions using the Editor Subsystem
+#if WITH_EDITOR
+	UMetaSoundEditorSubsystem* EditorSubsystem = GEditor ? GEditor->GetEditorSubsystem<UMetaSoundEditorSubsystem>() : nullptr;
+	if (EditorSubsystem)
+	{
+		for (const auto& PosPair : NewPositions)
+		{
+			FMetaSoundNodeHandle NodeHandle;
+			NodeHandle.NodeID = PosPair.Key;
+			
+			EMetaSoundBuilderResult LocationResult;
+			EditorSubsystem->SetNodeLocation(Builder, NodeHandle, PosPair.Value, LocationResult);
+		}
+	}
+#endif
+	
+	MetaSound->MarkPackageDirty();
+	SynchronizeMetaSoundEditorGraph(MetaSound, Builder);
+	
+	// Build result
+	TSharedPtr<FJsonObject> Result = MakeResult();
+	Result->SetStringField(TEXT("asset_path"), AssetPath);
+	Result->SetStringField(TEXT("strategy"), Strategy);
+	Result->SetNumberField(TEXT("nodes_positioned"), NodesPositioned);
+	Result->SetNumberField(TEXT("max_depth"), MaxDepth);
+	Result->SetNumberField(TEXT("spacing_x"), SpacingX);
+	Result->SetNumberField(TEXT("spacing_y"), SpacingY);
+	
+	// Include positioned node info
+	TArray<TSharedPtr<FJsonValue>> PositionedNodesArray;
+	for (const auto& PosPair : NewPositions)
+	{
+		TSharedPtr<FJsonObject> NodeInfo = MakeShared<FJsonObject>();
+		NodeInfo->SetStringField(TEXT("node_id"), GetNodeIdString(PosPair.Key));
+		NodeInfo->SetNumberField(TEXT("x"), PosPair.Value.X);
+		NodeInfo->SetNumberField(TEXT("y"), PosPair.Value.Y);
+		NodeInfo->SetNumberField(TEXT("depth"), NodeDepths.FindRef(PosPair.Key));
+		PositionedNodesArray.Add(MakeShared<FJsonValueObject>(NodeInfo));
+	}
+	Result->SetArrayField(TEXT("positioned_nodes"), PositionedNodesArray);
+	
+	return FECACommandResult::Success(Result);
+}
