ECABridge experimental MCP plugin

diff --git a/Engine/Plugins/Experimental/ECABridge/Source/ECABridge/Private/Commands/ECANiagaraCommands.cpp b/Engine/Plugins/Experimental/ECABridge/Source/ECABridge/Private/Commands/ECANiagaraCommands.cpp
new file mode 100644
--- /dev/null
+++ b/Engine/Plugins/Experimental/ECABridge/Source/ECABridge/Private/Commands/ECANiagaraCommands.cpp
@@ -0,0 +1,4202 @@
+// Copyright Epic Games, Inc. All Rights Reserved.
+
+#include "Commands/ECANiagaraCommands.h"
+#include "Commands/ECACommand.h"
+#include "Editor.h"
+#include "Engine/World.h"
+#include "EngineUtils.h"
+#include "AssetRegistry/AssetRegistryModule.h"
+#include "NiagaraSystem.h"
+#include "NiagaraActor.h"
+#include "NiagaraComponent.h"
+#include "NiagaraFunctionLibrary.h"
+#include "NiagaraDataInterfaceArrayFunctionLibrary.h"
+#include "NiagaraParameterStore.h"
+#include "NiagaraTypes.h"
+#include "NiagaraEmitter.h"
+#include "NiagaraEmitterHandle.h"
+#include "NiagaraScript.h"
+#include "NiagaraScriptSource.h"
+#include "NiagaraGraph.h"
+#include "NiagaraNodeOutput.h"
+#include "NiagaraNodeFunctionCall.h"
+#include "ViewModels/Stack/NiagaraStackGraphUtilities.h"
+#include "NiagaraSystemEditorData.h"
+#include "NiagaraSpriteRendererProperties.h"
+#include "NiagaraMeshRendererProperties.h"
+#include "NiagaraRibbonRendererProperties.h"
+#include "NiagaraLightRendererProperties.h"
+#include "NiagaraComponentRendererProperties.h"
+#include "NiagaraDataInterfaceCurve.h"
+#include "NiagaraDataInterfaceVectorCurve.h"
+#include "NiagaraDataInterfaceColorCurve.h"
+#include "Stateless/NiagaraStatelessDistribution.h"
+#include "Materials/MaterialInterface.h"
+#include "Engine/StaticMesh.h"
+#include "Curves/RichCurve.h"
+#include "ViewModels/Stack/NiagaraParameterHandle.h"
+#include "NiagaraScriptVariable.h"
+#include "NiagaraNodeCustomHlsl.h"
+#include "EdGraphSchema_Niagara.h"
+
+#include "EdGraph/EdGraphPin.h"
+#include "AssetToolsModule.h"
+#include "IAssetTools.h"
+#include "FileHelpers.h"
+#include "UObject/SavePackage.h"
+#include "Factories/Factory.h"
+#include "ScopedTransaction.h"
+#include "NiagaraNodeInput.h"
+#include "NiagaraEditorUtilities.h"
+#include "ObjectTools.h"
+
+// Register all Niagara commands
+REGISTER_ECA_COMMAND(FECACommand_GetNiagaraSystems)
+REGISTER_ECA_COMMAND(FECACommand_SpawnNiagaraEffect)
+REGISTER_ECA_COMMAND(FECACommand_CreateNiagaraSystem)
+REGISTER_ECA_COMMAND(FECACommand_GetNiagaraActors)
+REGISTER_ECA_COMMAND(FECACommand_SetNiagaraParameter)
+REGISTER_ECA_COMMAND(FECACommand_ControlNiagaraEffect)
+REGISTER_ECA_COMMAND(FECACommand_GetNiagaraParameters)
+REGISTER_ECA_COMMAND(FECACommand_AddNiagaraComponent)
+REGISTER_ECA_COMMAND(FECACommand_AddNiagaraEmitter)
+REGISTER_ECA_COMMAND(FECACommand_AddNiagaraModule)
+REGISTER_ECA_COMMAND(FECACommand_SetNiagaraModuleInput)
+REGISTER_ECA_COMMAND(FECACommand_GetNiagaraEmitters)
+REGISTER_ECA_COMMAND(FECACommand_GetNiagaraModules)
+REGISTER_ECA_COMMAND(FECACommand_ListNiagaraModules)
+REGISTER_ECA_COMMAND(FECACommand_RemoveNiagaraModule)
+REGISTER_ECA_COMMAND(FECACommand_DeleteNiagaraEmitter)
+REGISTER_ECA_COMMAND(FECACommand_SetNiagaraEmitterProperty)
+REGISTER_ECA_COMMAND(FECACommand_ListNiagaraEmitterTemplates)
+REGISTER_ECA_COMMAND(FECACommand_CreateNiagaraEmitter)
+REGISTER_ECA_COMMAND(FECACommand_AddNiagaraRenderer)
+REGISTER_ECA_COMMAND(FECACommand_SetNiagaraMaterial)
+REGISTER_ECA_COMMAND(FECACommand_SetNiagaraCurve)
+REGISTER_ECA_COMMAND(FECACommand_SetNiagaraDynamicInput)
+
+//------------------------------------------------------------------------------
+// Helper functions
+//------------------------------------------------------------------------------
+
+// Helper to remove nodes connected to an override pin (replacement for unexported RemoveNodesForStackFunctionInputOverridePin)
+static void RemoveOverridePinConnections(UEdGraphPin& OverridePin, UNiagaraGraph* Graph)
+{
+	if (OverridePin.LinkedTo.Num() == 0)
+	{
+		return;
+	}
+	
+	// Collect nodes to remove (nodes connected to this override pin)
+	TArray<UEdGraphNode*> NodesToRemove;
+	for (UEdGraphPin* LinkedPin : OverridePin.LinkedTo)
+	{
+		if (LinkedPin && LinkedPin->GetOwningNode())
+		{
+			NodesToRemove.AddUnique(LinkedPin->GetOwningNode());
+		}
+	}
+	
+	// Break all pin links first
+	OverridePin.BreakAllPinLinks();
+	
+	// Remove the connected nodes from the graph
+	for (UEdGraphNode* NodeToRemove : NodesToRemove)
+	{
+		if (NodeToRemove && Graph)
+		{
+			// Break all links on the node first
+			for (UEdGraphPin* Pin : NodeToRemove->Pins)
+			{
+				if (Pin)
+				{
+					Pin->BreakAllPinLinks();
+				}
+			}
+			Graph->RemoveNode(NodeToRemove);
+		}
+	}
+}
+
+static ANiagaraActor* FindNiagaraActorByName(const FString& ActorName)
+{
+	UWorld* World = GEditor->GetEditorWorldContext().World();
+	if (!World)
+	{
+		return nullptr;
+	}
+	
+	for (TActorIterator<ANiagaraActor> It(World); It; ++It)
+	{
+		ANiagaraActor* Actor = *It;
+		if (Actor->GetActorLabel().Equals(ActorName, ESearchCase::IgnoreCase) ||
+			Actor->GetName().Equals(ActorName, ESearchCase::IgnoreCase))
+		{
+			return Actor;
+		}
+	}
+	
+	return nullptr;
+}
+
+//------------------------------------------------------------------------------
+// GetNiagaraSystems
+//------------------------------------------------------------------------------
+
+FECACommandResult FECACommand_GetNiagaraSystems::Execute(const TSharedPtr<FJsonObject>& Params)
+{
+	FString PathFilter;
+	FString NameFilter;
+	GetStringParam(Params, TEXT("path_filter"), PathFilter, false);
+	GetStringParam(Params, TEXT("name_filter"), NameFilter, false);
+	
+	FAssetRegistryModule& AssetRegistryModule = FModuleManager::LoadModuleChecked<FAssetRegistryModule>("AssetRegistry");
+	IAssetRegistry& AssetRegistry = AssetRegistryModule.Get();
+	
+	TArray<FAssetData> AssetDataList;
+	AssetRegistry.GetAssetsByClass(UNiagaraSystem::StaticClass()->GetClassPathName(), AssetDataList);
+	
+	TArray<TSharedPtr<FJsonValue>> SystemsArray;
+	
+	for (const FAssetData& AssetData : AssetDataList)
+	{
+		FString AssetPath = AssetData.GetObjectPathString();
+		FString AssetName = AssetData.AssetName.ToString();
+		
+		// Apply filters
+		if (!PathFilter.IsEmpty() && !AssetPath.Contains(PathFilter, ESearchCase::IgnoreCase))
+		{
+			continue;
+		}
+		if (!NameFilter.IsEmpty() && !AssetName.Contains(NameFilter, ESearchCase::IgnoreCase))
+		{
+			continue;
+		}
+		
+		TSharedPtr<FJsonObject> SystemJson = MakeShared<FJsonObject>();
+		SystemJson->SetStringField(TEXT("name"), AssetName);
+		SystemJson->SetStringField(TEXT("path"), AssetPath);
+		SystemJson->SetStringField(TEXT("package_path"), AssetData.PackagePath.ToString());
+		
+		SystemsArray.Add(MakeShared<FJsonValueObject>(SystemJson));
+	}
+	
+	TSharedPtr<FJsonObject> Result = MakeResult();
+	Result->SetArrayField(TEXT("systems"), SystemsArray);
+	Result->SetNumberField(TEXT("count"), SystemsArray.Num());
+	
+	return FECACommandResult::Success(Result);
+}
+
+//------------------------------------------------------------------------------
+// SpawnNiagaraEffect
+//------------------------------------------------------------------------------
+
+FECACommandResult FECACommand_SpawnNiagaraEffect::Execute(const TSharedPtr<FJsonObject>& Params)
+{
+	FString SystemPath;
+	if (!GetStringParam(Params, TEXT("system_path"), SystemPath))
+	{
+		return FECACommandResult::Error(TEXT("Missing required parameter: system_path"));
+	}
+	
+	// Load the Niagara system
+	UNiagaraSystem* NiagaraSystem = LoadObject<UNiagaraSystem>(nullptr, *SystemPath);
+	if (!NiagaraSystem)
+	{
+		return FECACommandResult::Error(FString::Printf(TEXT("Failed to load Niagara system: %s"), *SystemPath));
+	}
+	
+	UWorld* World = GetEditorWorld();
+	if (!World)
+	{
+		return FECACommandResult::Error(TEXT("No world available"));
+	}
+	
+	// Parse transform parameters
+	FVector Location = FVector::ZeroVector;
+	FRotator Rotation = FRotator::ZeroRotator;
+	FVector Scale = FVector::OneVector;
+	bool bAutoActivate = true;
+	
+	GetVectorParam(Params, TEXT("location"), Location, false);
+	GetRotatorParam(Params, TEXT("rotation"), Rotation, false);
+	GetVectorParam(Params, TEXT("scale"), Scale, false);
+	GetBoolParam(Params, TEXT("auto_activate"), bAutoActivate, false);
+	
+	// Spawn the Niagara actor
+	FActorSpawnParameters SpawnParams;
+	SpawnParams.SpawnCollisionHandlingOverride = ESpawnActorCollisionHandlingMethod::AlwaysSpawn;
+	
+	ANiagaraActor* NiagaraActor = World->SpawnActor<ANiagaraActor>(Location, Rotation, SpawnParams);
+	if (!NiagaraActor)
+	{
+		return FECACommandResult::Error(TEXT("Failed to spawn Niagara actor"));
+	}
+	
+	// Configure the Niagara component
+	UNiagaraComponent* NiagaraComponent = NiagaraActor->GetNiagaraComponent();
+	if (NiagaraComponent)
+	{
+		NiagaraComponent->SetAsset(NiagaraSystem);
+		NiagaraComponent->SetAutoActivate(bAutoActivate);
+		NiagaraActor->SetActorScale3D(Scale);
+		
+		if (bAutoActivate)
+		{
+			NiagaraComponent->Activate(true);
+		}
+	}
+	
+	// Set name if specified
+	FString ActorName;
+	if (GetStringParam(Params, TEXT("name"), ActorName, false))
+	{
+		NiagaraActor->SetActorLabel(*ActorName);
+	}
+	
+	// Put in folder if specified
+	FString FolderPath;
+	if (GetStringParam(Params, TEXT("folder"), FolderPath, false))
+	{
+		NiagaraActor->SetFolderPath(*FolderPath);
+	}
+	
+	TSharedPtr<FJsonObject> Result = MakeResult();
+	Result->SetStringField(TEXT("actor_name"), NiagaraActor->GetActorLabel());
+	Result->SetStringField(TEXT("internal_name"), NiagaraActor->GetName());
+	Result->SetStringField(TEXT("system"), SystemPath);
+	Result->SetObjectField(TEXT("transform"), TransformToJson(NiagaraActor->GetTransform()));
+	Result->SetBoolField(TEXT("auto_activate"), bAutoActivate);
+	
+	return FECACommandResult::Success(Result);
+}
+
+//------------------------------------------------------------------------------
+// CreateNiagaraSystem
+//------------------------------------------------------------------------------
+
+FECACommandResult FECACommand_CreateNiagaraSystem::Execute(const TSharedPtr<FJsonObject>& Params)
+{
+	FString AssetPath;
+	if (!GetStringParam(Params, TEXT("asset_path"), AssetPath))
+	{
+		return FECACommandResult::Error(TEXT("Missing required parameter: asset_path"));
+	}
+	
+	FString Template = TEXT("empty");
+	GetStringParam(Params, TEXT("template"), Template, false);
+	
+	// Parse path into package and asset name
+	FString PackagePath;
+	FString AssetName;
+	
+	int32 LastSlash;
+	if (AssetPath.FindLastChar('/', LastSlash))
+	{
+		PackagePath = AssetPath.Left(LastSlash);
+		AssetName = AssetPath.Mid(LastSlash + 1);
+	}
+	else
+	{
+		return FECACommandResult::Error(TEXT("Invalid asset path format. Expected format: /Game/Path/AssetName"));
+	}
+	
+	// Remove NS_ prefix for cleaner asset name if present
+	FString CleanAssetName = AssetName;
+	if (!AssetName.StartsWith(TEXT("NS_")))
+	{
+		// Add NS_ prefix if not present for consistency
+		AssetName = FString::Printf(TEXT("NS_%s"), *CleanAssetName);
+	}
+	
+	// Create the package
+	FString FullPackagePath = PackagePath / AssetName;
+	UPackage* Package = CreatePackage(*FullPackagePath);
+	if (!Package)
+	{
+		return FECACommandResult::Error(FString::Printf(TEXT("Failed to create package at: %s"), *FullPackagePath));
+	}
+	
+	Package->FullyLoad();
+	
+	// Create the Niagara system using AssetTools
+	IAssetTools& AssetTools = FModuleManager::LoadModuleChecked<FAssetToolsModule>("AssetTools").Get();
+	
+	// Find and use the Niagara system factory
+	TArray<UFactory*> Factories = AssetTools.GetNewAssetFactories();
+	UFactory* NiagaraFactory = nullptr;
+	
+	for (UFactory* Factory : Factories)
+	{
+		if (Factory && Factory->SupportedClass == UNiagaraSystem::StaticClass())
+		{
+			NiagaraFactory = Factory;
+			break;
+		}
+	}
+	
+	UObject* NewAsset = nullptr;
+	
+	if (NiagaraFactory)
+	{
+		NewAsset = AssetTools.CreateAsset(AssetName, PackagePath, UNiagaraSystem::StaticClass(), NiagaraFactory);
+	}
+	else
+	{
+		// Fallback: Create directly without factory
+		NewAsset = NewObject<UNiagaraSystem>(Package, FName(*AssetName), RF_Public | RF_Standalone);
+	}
+	
+	if (!NewAsset)
+	{
+		return FECACommandResult::Error(FString::Printf(TEXT("Failed to create Niagara system at: %s"), *AssetPath));
+	}
+	
+	// Mark package dirty
+	Package->MarkPackageDirty();
+	
+	// Notify asset registry
+	FAssetRegistryModule::AssetCreated(NewAsset);
+	
+	// Save the asset
+	FString PackageFileName = FPackageName::LongPackageNameToFilename(FullPackagePath, FPackageName::GetAssetPackageExtension());
+	FSavePackageArgs SaveArgs;
+	SaveArgs.TopLevelFlags = RF_Public | RF_Standalone;
+	UPackage::SavePackage(Package, NewAsset, *PackageFileName, SaveArgs);
+	
+	TSharedPtr<FJsonObject> Result = MakeResult();
+	Result->SetStringField(TEXT("asset_name"), NewAsset->GetName());
+	Result->SetStringField(TEXT("asset_path"), NewAsset->GetPathName());
+	Result->SetStringField(TEXT("package_path"), PackagePath);
+	Result->SetStringField(TEXT("template"), Template);
+	
+	return FECACommandResult::Success(Result);
+}
+
+//------------------------------------------------------------------------------
+// GetNiagaraActors
+//------------------------------------------------------------------------------
+
+FECACommandResult FECACommand_GetNiagaraActors::Execute(const TSharedPtr<FJsonObject>& Params)
+{
+	UWorld* World = GetEditorWorld();
+	if (!World)
+	{
+		return FECACommandResult::Error(TEXT("No world available"));
+	}
+	
+	FString SystemFilter;
+	GetStringParam(Params, TEXT("system_filter"), SystemFilter, false);
+	
+	TArray<TSharedPtr<FJsonValue>> ActorsArray;
+	
+	for (TActorIterator<ANiagaraActor> It(World); It; ++It)
+	{
+		ANiagaraActor* Actor = *It;
+		UNiagaraComponent* NiagaraComp = Actor->GetNiagaraComponent();
+		
+		FString SystemName;
+		FString SystemPath;
+		
+		if (NiagaraComp && NiagaraComp->GetAsset())
+		{
+			SystemName = NiagaraComp->GetAsset()->GetName();
+			SystemPath = NiagaraComp->GetAsset()->GetPathName();
+			
+			// Apply filter
+			if (!SystemFilter.IsEmpty() && !SystemName.Contains(SystemFilter, ESearchCase::IgnoreCase))
+			{
+				continue;
+			}
+		}
+		else if (!SystemFilter.IsEmpty())
+		{
+			// Skip actors without systems when filtering
+			continue;
+		}
+		
+		TSharedPtr<FJsonObject> ActorJson = MakeShared<FJsonObject>();
+		ActorJson->SetStringField(TEXT("name"), Actor->GetActorLabel());
+		ActorJson->SetStringField(TEXT("internal_name"), Actor->GetName());
+		ActorJson->SetStringField(TEXT("system_name"), SystemName);
+		ActorJson->SetStringField(TEXT("system_path"), SystemPath);
+		ActorJson->SetObjectField(TEXT("location"), VectorToJson(Actor->GetActorLocation()));
+		ActorJson->SetObjectField(TEXT("rotation"), RotatorToJson(Actor->GetActorRotation()));
+		ActorJson->SetObjectField(TEXT("scale"), VectorToJson(Actor->GetActorScale3D()));
+		
+		if (NiagaraComp)
+		{
+			ActorJson->SetBoolField(TEXT("is_active"), NiagaraComp->IsActive());
+			ActorJson->SetBoolField(TEXT("auto_activate"), NiagaraComp->bAutoActivate);
+		}
+		
+		// Include folder path
+		FName FolderPath = Actor->GetFolderPath();
+		if (!FolderPath.IsNone())
+		{
+			ActorJson->SetStringField(TEXT("folder"), FolderPath.ToString());
+		}
+		
+		ActorsArray.Add(MakeShared<FJsonValueObject>(ActorJson));
+	}
+	
+	TSharedPtr<FJsonObject> Result = MakeResult();
+	Result->SetArrayField(TEXT("actors"), ActorsArray);
+	Result->SetNumberField(TEXT("count"), ActorsArray.Num());
+	
+	return FECACommandResult::Success(Result);
+}
+
+//------------------------------------------------------------------------------
+// SetNiagaraParameter
+//------------------------------------------------------------------------------
+
+FECACommandResult FECACommand_SetNiagaraParameter::Execute(const TSharedPtr<FJsonObject>& Params)
+{
+	FString ActorName;
+	if (!GetStringParam(Params, TEXT("actor_name"), ActorName))
+	{
+		return FECACommandResult::Error(TEXT("Missing required parameter: actor_name"));
+	}
+	
+	FString ParameterName;
+	if (!GetStringParam(Params, TEXT("parameter_name"), ParameterName))
+	{
+		return FECACommandResult::Error(TEXT("Missing required parameter: parameter_name"));
+	}
+	
+	FString ParameterType;
+	if (!GetStringParam(Params, TEXT("parameter_type"), ParameterType))
+	{
+		return FECACommandResult::Error(TEXT("Missing required parameter: parameter_type"));
+	}
+	
+	ANiagaraActor* NiagaraActor = FindNiagaraActorByName(ActorName);
+	if (!NiagaraActor)
+	{
+		return FECACommandResult::Error(FString::Printf(TEXT("Niagara actor not found: %s"), *ActorName));
+	}
+	
+	UNiagaraComponent* NiagaraComp = NiagaraActor->GetNiagaraComponent();
+	if (!NiagaraComp)
+	{
+		return FECACommandResult::Error(TEXT("Niagara component not found"));
+	}
+	
+	TSharedPtr<FJsonObject> Result = MakeResult();
+	Result->SetStringField(TEXT("actor_name"), ActorName);
+	Result->SetStringField(TEXT("parameter_name"), ParameterName);
+	Result->SetStringField(TEXT("parameter_type"), ParameterType);
+	
+	// Set parameter based on type
+	if (ParameterType.Equals(TEXT("float"), ESearchCase::IgnoreCase))
+	{
+		double Value;
+		if (!GetFloatParam(Params, TEXT("value"), Value))
+		{
+			return FECACommandResult::Error(TEXT("Missing or invalid value for float parameter"));
+		}
+		NiagaraComp->SetVariableFloat(FName(*ParameterName), static_cast<float>(Value));
+		Result->SetNumberField(TEXT("value"), Value);
+	}
+	else if (ParameterType.Equals(TEXT("int"), ESearchCase::IgnoreCase))
+	{
+		int32 Value;
+		if (!GetIntParam(Params, TEXT("value"), Value))
+		{
+			return FECACommandResult::Error(TEXT("Missing or invalid value for int parameter"));
+		}
+		NiagaraComp->SetVariableInt(FName(*ParameterName), Value);
+		Result->SetNumberField(TEXT("value"), Value);
+	}
+	else if (ParameterType.Equals(TEXT("bool"), ESearchCase::IgnoreCase))
+	{
+		bool Value;
+		if (!GetBoolParam(Params, TEXT("value"), Value))
+		{
+			return FECACommandResult::Error(TEXT("Missing or invalid value for bool parameter"));
+		}
+		NiagaraComp->SetVariableBool(FName(*ParameterName), Value);
+		Result->SetBoolField(TEXT("value"), Value);
+	}
+	else if (ParameterType.Equals(TEXT("vector"), ESearchCase::IgnoreCase))
+	{
+		FVector Value;
+		if (!GetVectorParam(Params, TEXT("value"), Value))
+		{
+			return FECACommandResult::Error(TEXT("Missing or invalid value for vector parameter"));
+		}
+		NiagaraComp->SetVariableVec3(FName(*ParameterName), Value);
+		Result->SetObjectField(TEXT("value"), VectorToJson(Value));
+	}
+	else if (ParameterType.Equals(TEXT("color"), ESearchCase::IgnoreCase) || 
+	         ParameterType.Equals(TEXT("linear_color"), ESearchCase::IgnoreCase))
+	{
+		const TSharedPtr<FJsonObject>* ValueObj;
+		if (!GetObjectParam(Params, TEXT("value"), ValueObj))
+		{
+			return FECACommandResult::Error(TEXT("Missing or invalid value for color parameter"));
+		}
+		
+		FLinearColor Color;
+		Color.R = static_cast<float>((*ValueObj)->GetNumberField(TEXT("r")));
+		Color.G = static_cast<float>((*ValueObj)->GetNumberField(TEXT("g")));
+		Color.B = static_cast<float>((*ValueObj)->GetNumberField(TEXT("b")));
+		Color.A = (*ValueObj)->HasField(TEXT("a")) ? static_cast<float>((*ValueObj)->GetNumberField(TEXT("a"))) : 1.0f;
+		
+		NiagaraComp->SetVariableLinearColor(FName(*ParameterName), Color);
+		
+		TSharedPtr<FJsonObject> ColorJson = MakeShared<FJsonObject>();
+		ColorJson->SetNumberField(TEXT("r"), Color.R);
+		ColorJson->SetNumberField(TEXT("g"), Color.G);
+		ColorJson->SetNumberField(TEXT("b"), Color.B);
+		ColorJson->SetNumberField(TEXT("a"), Color.A);
+		Result->SetObjectField(TEXT("value"), ColorJson);
+	}
+	else
+	{
+		return FECACommandResult::Error(FString::Printf(TEXT("Unsupported parameter type: %s. Supported types: float, int, bool, vector, color, linear_color"), *ParameterType));
+	}
+	
+	return FECACommandResult::Success(Result);
+}
+
+//------------------------------------------------------------------------------
+// ControlNiagaraEffect
+//------------------------------------------------------------------------------
+
+FECACommandResult FECACommand_ControlNiagaraEffect::Execute(const TSharedPtr<FJsonObject>& Params)
+{
+	FString ActorName;
+	if (!GetStringParam(Params, TEXT("actor_name"), ActorName))
+	{
+		return FECACommandResult::Error(TEXT("Missing required parameter: actor_name"));
+	}
+	
+	FString Action;
+	if (!GetStringParam(Params, TEXT("action"), Action))
+	{
+		return FECACommandResult::Error(TEXT("Missing required parameter: action"));
+	}
+	
+	ANiagaraActor* NiagaraActor = FindNiagaraActorByName(ActorName);
+	if (!NiagaraActor)
+	{
+		return FECACommandResult::Error(FString::Printf(TEXT("Niagara actor not found: %s"), *ActorName));
+	}
+	
+	UNiagaraComponent* NiagaraComp = NiagaraActor->GetNiagaraComponent();
+	if (!NiagaraComp)
+	{
+		return FECACommandResult::Error(TEXT("Niagara component not found"));
+	}
+	
+	TSharedPtr<FJsonObject> Result = MakeResult();
+	Result->SetStringField(TEXT("actor_name"), ActorName);
+	Result->SetStringField(TEXT("action"), Action);
+	
+	if (Action.Equals(TEXT("activate"), ESearchCase::IgnoreCase))
+	{
+		NiagaraComp->Activate(true);
+		Result->SetBoolField(TEXT("is_active"), true);
+	}
+	else if (Action.Equals(TEXT("deactivate"), ESearchCase::IgnoreCase))
+	{
+		NiagaraComp->Deactivate();
+		Result->SetBoolField(TEXT("is_active"), false);
+	}
+	else if (Action.Equals(TEXT("reset"), ESearchCase::IgnoreCase))
+	{
+		NiagaraComp->ResetSystem();
+		Result->SetBoolField(TEXT("is_active"), NiagaraComp->IsActive());
+	}
+	else if (Action.Equals(TEXT("reset_system"), ESearchCase::IgnoreCase))
+	{
+		NiagaraComp->ReinitializeSystem();
+		Result->SetBoolField(TEXT("is_active"), NiagaraComp->IsActive());
+	}
+	else
+	{
+		return FECACommandResult::Error(FString::Printf(TEXT("Unknown action: %s. Supported actions: activate, deactivate, reset, reset_system"), *Action));
+	}
+	
+	return FECACommandResult::Success(Result);
+}
+
+//------------------------------------------------------------------------------
+// GetNiagaraParameters
+//------------------------------------------------------------------------------
+
+FECACommandResult FECACommand_GetNiagaraParameters::Execute(const TSharedPtr<FJsonObject>& Params)
+{
+	UNiagaraSystem* NiagaraSystem = nullptr;
+	FString SourceName;
+	
+	// Try to get from actor first
+	FString ActorName;
+	if (GetStringParam(Params, TEXT("actor_name"), ActorName, false))
+	{
+		ANiagaraActor* NiagaraActor = FindNiagaraActorByName(ActorName);
+		if (NiagaraActor)
+		{
+			UNiagaraComponent* NiagaraComp = NiagaraActor->GetNiagaraComponent();
+			if (NiagaraComp)
+			{
+				NiagaraSystem = NiagaraComp->GetAsset();
+				SourceName = ActorName;
+			}
+		}
+		
+		if (!NiagaraSystem)
+		{
+			return FECACommandResult::Error(FString::Printf(TEXT("Niagara actor not found or has no system: %s"), *ActorName));
+		}
+	}
+	else
+	{
+		// Try to load from path
+		FString SystemPath;
+		if (GetStringParam(Params, TEXT("system_path"), SystemPath, false))
+		{
+			NiagaraSystem = LoadObject<UNiagaraSystem>(nullptr, *SystemPath);
+			SourceName = SystemPath;
+			
+			if (!NiagaraSystem)
+			{
+				return FECACommandResult::Error(FString::Printf(TEXT("Failed to load Niagara system: %s"), *SystemPath));
+			}
+		}
+		else
+		{
+			return FECACommandResult::Error(TEXT("Must specify either actor_name or system_path"));
+		}
+	}
+	
+	TArray<TSharedPtr<FJsonValue>> ParametersArray;
+	
+	// Get user-exposed parameters from the system
+	const FNiagaraUserRedirectionParameterStore& ExposedParams = NiagaraSystem->GetExposedParameters();
+	TArrayView<const FNiagaraVariableWithOffset> UserParameters = ExposedParams.ReadParameterVariables();
+	
+	for (const FNiagaraVariableWithOffset& VarWithOffset : UserParameters)
+	{
+		const FNiagaraVariableBase& Var = VarWithOffset;
+		TSharedPtr<FJsonObject> ParamJson = MakeShared<FJsonObject>();
+		ParamJson->SetStringField(TEXT("name"), Var.GetName().ToString());
+		ParamJson->SetStringField(TEXT("type"), Var.GetType().GetName());
+		
+		// Try to get default value information
+		FString TypeName = Var.GetType().GetName();
+		if (TypeName.Contains(TEXT("Float")))
+		{
+			ParamJson->SetStringField(TEXT("value_type"), TEXT("float"));
+		}
+		else if (TypeName.Contains(TEXT("Int")))
+		{
+			ParamJson->SetStringField(TEXT("value_type"), TEXT("int"));
+		}
+		else if (TypeName.Contains(TEXT("Bool")))
+		{
+			ParamJson->SetStringField(TEXT("value_type"), TEXT("bool"));
+		}
+		else if (TypeName.Contains(TEXT("Vector")) || TypeName.Contains(TEXT("Position")))
+		{
+			ParamJson->SetStringField(TEXT("value_type"), TEXT("vector"));
+		}
+		else if (TypeName.Contains(TEXT("Color")))
+		{
+			ParamJson->SetStringField(TEXT("value_type"), TEXT("color"));
+		}
+		else
+		{
+			ParamJson->SetStringField(TEXT("value_type"), TEXT("other"));
+		}
+		
+		ParametersArray.Add(MakeShared<FJsonValueObject>(ParamJson));
+	}
+	
+	TSharedPtr<FJsonObject> Result = MakeResult();
+	Result->SetStringField(TEXT("source"), SourceName);
+	Result->SetStringField(TEXT("system_name"), NiagaraSystem->GetName());
+	Result->SetArrayField(TEXT("parameters"), ParametersArray);
+	Result->SetNumberField(TEXT("count"), ParametersArray.Num());
+	
+	return FECACommandResult::Success(Result);
+}
+
+//------------------------------------------------------------------------------
+// AddNiagaraComponent
+//------------------------------------------------------------------------------
+
+FECACommandResult FECACommand_AddNiagaraComponent::Execute(const TSharedPtr<FJsonObject>& Params)
+{
+	FString ActorName;
+	if (!GetStringParam(Params, TEXT("actor_name"), ActorName))
+	{
+		return FECACommandResult::Error(TEXT("Missing required parameter: actor_name"));
+	}
+	
+	FString SystemPath;
+	if (!GetStringParam(Params, TEXT("system_path"), SystemPath))
+	{
+		return FECACommandResult::Error(TEXT("Missing required parameter: system_path"));
+	}
+	
+	AActor* Actor = FindActorByName(ActorName);
+	if (!Actor)
+	{
+		return FECACommandResult::Error(FString::Printf(TEXT("Actor not found: %s"), *ActorName));
+	}
+	
+	UNiagaraSystem* NiagaraSystem = LoadObject<UNiagaraSystem>(nullptr, *SystemPath);
+	if (!NiagaraSystem)
+	{
+		return FECACommandResult::Error(FString::Printf(TEXT("Failed to load Niagara system: %s"), *SystemPath));
+	}
+	
+	// Create the Niagara component
+	FString ComponentName = TEXT("NiagaraComponent");
+	GetStringParam(Params, TEXT("component_name"), ComponentName, false);
+	
+	UNiagaraComponent* NiagaraComp = NewObject<UNiagaraComponent>(Actor, FName(*ComponentName));
+	if (!NiagaraComp)
+	{
+		return FECACommandResult::Error(TEXT("Failed to create Niagara component"));
+	}
+	
+	NiagaraComp->SetAsset(NiagaraSystem);
+	
+	// Set relative transform
+	FVector RelativeLocation = FVector::ZeroVector;
+	FRotator RelativeRotation = FRotator::ZeroRotator;
+	GetVectorParam(Params, TEXT("relative_location"), RelativeLocation, false);
+	GetRotatorParam(Params, TEXT("relative_rotation"), RelativeRotation, false);
+	
+	NiagaraComp->SetRelativeLocation(RelativeLocation);
+	NiagaraComp->SetRelativeRotation(RelativeRotation);
+	
+	// Auto-activate
+	bool bAutoActivate = true;
+	GetBoolParam(Params, TEXT("auto_activate"), bAutoActivate, false);
+	NiagaraComp->SetAutoActivate(bAutoActivate);
+	
+	// Attach to actor
+	NiagaraComp->AttachToComponent(Actor->GetRootComponent(), FAttachmentTransformRules::KeepRelativeTransform);
+	NiagaraComp->RegisterComponent();
+	Actor->AddInstanceComponent(NiagaraComp);
+	
+	if (bAutoActivate)
+	{
+		NiagaraComp->Activate(true);
+	}
+	
+	TSharedPtr<FJsonObject> Result = MakeResult();
+	Result->SetStringField(TEXT("actor_name"), ActorName);
+	Result->SetStringField(TEXT("component_name"), NiagaraComp->GetName());
+	Result->SetStringField(TEXT("system_path"), SystemPath);
+	Result->SetObjectField(TEXT("relative_location"), VectorToJson(RelativeLocation));
+	Result->SetObjectField(TEXT("relative_rotation"), RotatorToJson(RelativeRotation));
+	Result->SetBoolField(TEXT("auto_activate"), bAutoActivate);
+	Result->SetBoolField(TEXT("is_active"), NiagaraComp->IsActive());
+	
+	return FECACommandResult::Success(Result);
+}
+
+
+//------------------------------------------------------------------------------
+// Helper: Parse script usage string to enum
+//------------------------------------------------------------------------------
+
+static ENiagaraScriptUsage ParseScriptUsage(const FString& UsageString)
+{
+	if (UsageString.Equals(TEXT("emitter_spawn"), ESearchCase::IgnoreCase))
+	{
+		return ENiagaraScriptUsage::EmitterSpawnScript;
+	}
+	else if (UsageString.Equals(TEXT("emitter_update"), ESearchCase::IgnoreCase))
+	{
+		return ENiagaraScriptUsage::EmitterUpdateScript;
+	}
+	else if (UsageString.Equals(TEXT("particle_spawn"), ESearchCase::IgnoreCase))
+	{
+		return ENiagaraScriptUsage::ParticleSpawnScript;
+	}
+	else if (UsageString.Equals(TEXT("particle_update"), ESearchCase::IgnoreCase))
+	{
+		return ENiagaraScriptUsage::ParticleUpdateScript;
+	}
+	else if (UsageString.Equals(TEXT("system_spawn"), ESearchCase::IgnoreCase))
+	{
+		return ENiagaraScriptUsage::SystemSpawnScript;
+	}
+	else if (UsageString.Equals(TEXT("system_update"), ESearchCase::IgnoreCase))
+	{
+		return ENiagaraScriptUsage::SystemUpdateScript;
+	}
+	return ENiagaraScriptUsage::ParticleUpdateScript; // Default
+}
+
+static FString ScriptUsageToString(ENiagaraScriptUsage Usage)
+{
+	switch (Usage)
+	{
+		case ENiagaraScriptUsage::EmitterSpawnScript: return TEXT("emitter_spawn");
+		case ENiagaraScriptUsage::EmitterUpdateScript: return TEXT("emitter_update");
+		case ENiagaraScriptUsage::ParticleSpawnScript: return TEXT("particle_spawn");
+		case ENiagaraScriptUsage::ParticleUpdateScript: return TEXT("particle_update");
+		case ENiagaraScriptUsage::SystemSpawnScript: return TEXT("system_spawn");
+		case ENiagaraScriptUsage::SystemUpdateScript: return TEXT("system_update");
+		default: return TEXT("unknown");
+	}
+}
+
+//------------------------------------------------------------------------------
+// Helper: Find emitter handle by name
+//------------------------------------------------------------------------------
+
+static FNiagaraEmitterHandle* FindEmitterHandleByName(UNiagaraSystem* System, const FString& EmitterName, int32* OutIndex = nullptr)
+{
+	if (!System)
+	{
+		return nullptr;
+	}
+	
+	TArray<FNiagaraEmitterHandle>& EmitterHandles = System->GetEmitterHandles();
+	
+	// Check if it's an index reference like "Emitter[0]"
+	if (EmitterName.StartsWith(TEXT("Emitter[")) && EmitterName.EndsWith(TEXT("]")))
+	{
+		FString IndexStr = EmitterName.Mid(8, EmitterName.Len() - 9);
+		int32 Index = FCString::Atoi(*IndexStr);
+		if (Index >= 0 && Index < EmitterHandles.Num())
+		{
+			if (OutIndex) *OutIndex = Index;
+			return &EmitterHandles[Index];
+		}
+		return nullptr;
+	}
+	
+	// Search by name
+	for (int32 i = 0; i < EmitterHandles.Num(); ++i)
+	{
+		FNiagaraEmitterHandle& Handle = EmitterHandles[i];
+		if (Handle.GetName().ToString().Equals(EmitterName, ESearchCase::IgnoreCase) ||
+			Handle.GetUniqueInstanceName().Equals(EmitterName, ESearchCase::IgnoreCase))
+		{
+			if (OutIndex) *OutIndex = i;
+			return &Handle;
+		}
+	}
+	
+	return nullptr;
+}
+
+//------------------------------------------------------------------------------
+// AddNiagaraEmitter
+//------------------------------------------------------------------------------
+
+FECACommandResult FECACommand_AddNiagaraEmitter::Execute(const TSharedPtr<FJsonObject>& Params)
+{
+#if WITH_EDITORONLY_DATA
+	FString SystemPath;
+	if (!GetStringParam(Params, TEXT("system_path"), SystemPath))
+	{
+		return FECACommandResult::Error(TEXT("Missing required parameter: system_path"));
+	}
+	
+	UNiagaraSystem* NiagaraSystem = LoadObject<UNiagaraSystem>(nullptr, *SystemPath);
+	if (!NiagaraSystem)
+	{
+		return FECACommandResult::Error(FString::Printf(TEXT("Failed to load Niagara system: %s"), *SystemPath));
+	}
+	
+	FString EmitterPath = TEXT("empty");
+	GetStringParam(Params, TEXT("emitter_path"), EmitterPath, false);
+	
+	FString EmitterName;
+	GetStringParam(Params, TEXT("emitter_name"), EmitterName, false);
+	
+	UNiagaraEmitter* SourceEmitter = nullptr;
+	
+	if (!EmitterPath.Equals(TEXT("empty"), ESearchCase::IgnoreCase))
+	{
+		// Load the specified emitter
+		SourceEmitter = LoadObject<UNiagaraEmitter>(nullptr, *EmitterPath);
+		if (!SourceEmitter)
+		{
+			return FECACommandResult::Error(FString::Printf(TEXT("Failed to load emitter: %s"), *EmitterPath));
+		}
+	}
+	else
+	{
+		// Try to load a default emitter template from the engine
+		// These are the standard Niagara emitter templates
+		const TCHAR* DefaultEmitterPaths[] = {
+			TEXT("/Niagara/DefaultAssets/Templates/Emitters/Fountain.Fountain"),
+			TEXT("/Niagara/DefaultAssets/Templates/Emitters/SimpleSpriteBurst.SimpleSpriteBurst"),
+			TEXT("/Niagara/DefaultAssets/Templates/Emitters/Minimal.Minimal"),
+			TEXT("/Niagara/DefaultAssets/Templates/Emitters/OmnidirectionalBurst.OmnidirectionalBurst"),
+		};
+		
+		for (const TCHAR* Path : DefaultEmitterPaths)
+		{
+			SourceEmitter = LoadObject<UNiagaraEmitter>(nullptr, Path);
+			if (SourceEmitter)
+			{
+				break;
+			}
+		}
+		
+		if (!SourceEmitter)
+		{
+			return FECACommandResult::Error(TEXT("Could not find default emitter template. Please specify an emitter_path. Available templates include: /Niagara/DefaultAssets/Templates/Emitters/Fountain, /Niagara/DefaultAssets/Templates/Emitters/SimpleSpriteBurst, /Niagara/DefaultAssets/Templates/Emitters/Minimal"));
+		}
+	}
+	
+	// Determine the emitter name
+	FName FinalEmitterName = EmitterName.IsEmpty() ? FName(*SourceEmitter->GetName()) : FName(*EmitterName);
+	
+	// Use the system's AddEmitterHandle method to properly add the emitter
+	FGuid EmitterVersion = SourceEmitter->GetExposedVersion().VersionGuid;
+	FNiagaraEmitterHandle NewHandle = NiagaraSystem->AddEmitterHandle(*SourceEmitter, FinalEmitterName, EmitterVersion);
+	
+	// Mark the system's graph as changed - NotifyGraphChanged() is exported from UEdGraph base class
+	// RequestCompile() below handles the actual recompilation
+	if (UNiagaraScriptSource* SystemScriptSource = Cast<UNiagaraScriptSource>(NiagaraSystem->GetSystemSpawnScript()->GetLatestSource()))
+	{
+		if (UNiagaraGraph* SystemGraph = SystemScriptSource->NodeGraph)
+		{
+			SystemGraph->NotifyGraphChanged();
+		}
+	}
+	
+	// CRITICAL: Synchronize the overview graph with the system - this creates the UI nodes in the editor!
+	// Without this call, the emitter exists in the data model but won't appear in the Niagara editor UI
+	// until another action (like manually adding an emitter) triggers a sync
+	UNiagaraSystemEditorData* SystemEditorData = Cast<UNiagaraSystemEditorData>(NiagaraSystem->GetEditorData());
+	if (SystemEditorData)
+	{
+		SystemEditorData->SynchronizeOverviewGraphWithSystem(*NiagaraSystem);
+	}
+	
+	// Mark the system as modified
+	NiagaraSystem->MarkPackageDirty();
+	
+	// Request recompile to make the emitter visible in the editor
+	NiagaraSystem->RequestCompile(false);
+	
+	// Notify that the system has changed - this triggers UI refresh
+	NiagaraSystem->PostEditChange();
+	
+	TSharedPtr<FJsonObject> Result = MakeResult();
+	Result->SetStringField(TEXT("system_path"), SystemPath);
+	Result->SetStringField(TEXT("emitter_name"), NewHandle.GetName().ToString());
+	Result->SetStringField(TEXT("emitter_id"), NewHandle.GetId().ToString());
+	Result->SetNumberField(TEXT("emitter_index"), NiagaraSystem->GetEmitterHandles().Num() - 1);
+	
+	return FECACommandResult::Success(Result);
+#else
+	return FECACommandResult::Error(TEXT("AddNiagaraEmitter is only available in editor builds"));
+#endif
+}
+
+//------------------------------------------------------------------------------
+// AddNiagaraModule
+//------------------------------------------------------------------------------
+
+FECACommandResult FECACommand_AddNiagaraModule::Execute(const TSharedPtr<FJsonObject>& Params)
+{
+	FString SystemPath;
+	if (!GetStringParam(Params, TEXT("system_path"), SystemPath))
+	{
+		return FECACommandResult::Error(TEXT("Missing required parameter: system_path"));
+	}
+	
+	FString EmitterName;
+	if (!GetStringParam(Params, TEXT("emitter_name"), EmitterName))
+	{
+		return FECACommandResult::Error(TEXT("Missing required parameter: emitter_name"));
+	}
+	
+	FString ModulePath;
+	if (!GetStringParam(Params, TEXT("module_path"), ModulePath))
+	{
+		return FECACommandResult::Error(TEXT("Missing required parameter: module_path"));
+	}
+	
+	FString ScriptUsageStr;
+	if (!GetStringParam(Params, TEXT("script_usage"), ScriptUsageStr))
+	{
+		return FECACommandResult::Error(TEXT("Missing required parameter: script_usage"));
+	}
+	
+	int32 TargetIndex = INDEX_NONE;
+	GetIntParam(Params, TEXT("index"), TargetIndex, false);
+	
+	// Load the system
+	UNiagaraSystem* NiagaraSystem = LoadObject<UNiagaraSystem>(nullptr, *SystemPath);
+	if (!NiagaraSystem)
+	{
+		return FECACommandResult::Error(FString::Printf(TEXT("Failed to load Niagara system: %s"), *SystemPath));
+	}
+	
+	// Find the emitter
+	int32 EmitterIndex;
+	FNiagaraEmitterHandle* EmitterHandle = FindEmitterHandleByName(NiagaraSystem, EmitterName, &EmitterIndex);
+	if (!EmitterHandle)
+	{
+		return FECACommandResult::Error(FString::Printf(TEXT("Emitter not found: %s"), *EmitterName));
+	}
+	
+	FVersionedNiagaraEmitterData* EmitterData = EmitterHandle->GetEmitterData();
+	if (!EmitterData)
+	{
+		return FECACommandResult::Error(TEXT("Emitter has no data"));
+	}
+	
+	// Load the module script
+	UNiagaraScript* ModuleScript = LoadObject<UNiagaraScript>(nullptr, *ModulePath);
+	if (!ModuleScript)
+	{
+		return FECACommandResult::Error(FString::Printf(TEXT("Failed to load module: %s"), *ModulePath));
+	}
+	
+	// Parse script usage
+	ENiagaraScriptUsage ScriptUsage = ParseScriptUsage(ScriptUsageStr);
+	
+	// Get the script source and graph
+	UNiagaraScriptSource* ScriptSource = Cast<UNiagaraScriptSource>(EmitterData->GraphSource);
+	if (!ScriptSource)
+	{
+		return FECACommandResult::Error(TEXT("Emitter has no script source"));
+	}
+	
+	UNiagaraGraph* Graph = ScriptSource->NodeGraph;
+	if (!Graph)
+	{
+		return FECACommandResult::Error(TEXT("Emitter has no graph"));
+	}
+	
+	// Find the output node for the target script usage
+	UNiagaraNodeOutput* OutputNode = Graph->FindEquivalentOutputNode(ScriptUsage, FGuid());
+	if (!OutputNode)
+	{
+		return FECACommandResult::Error(FString::Printf(TEXT("Could not find output node for script usage: %s"), *ScriptUsageStr));
+	}
+	
+	// Create a function call node for the module
+	UNiagaraNodeFunctionCall* NewModuleNode = NewObject<UNiagaraNodeFunctionCall>(Graph);
+	NewModuleNode->FunctionScript = ModuleScript;
+	NewModuleNode->CreateNewGuid();
+	NewModuleNode->PostPlacedNewNode();
+	NewModuleNode->AllocateDefaultPins();
+	Graph->AddNode(NewModuleNode, false, false);
+	
+	// Find the parameter map input pin on the output node and connect our module
+	UEdGraphPin* OutputNodeInputPin = nullptr;
+	for (UEdGraphPin* Pin : OutputNode->Pins)
+	{
+		if (Pin->Direction == EGPD_Input)
+		{
+			OutputNodeInputPin = Pin;
+			break;
+		}
+	}
+	
+	if (OutputNodeInputPin)
+	{
+		// Find the output pin on our new module
+		UEdGraphPin* ModuleOutputPin = nullptr;
+		for (UEdGraphPin* Pin : NewModuleNode->Pins)
+		{
+			if (Pin->Direction == EGPD_Output)
+			{
+				ModuleOutputPin = Pin;
+				break;
+			}
+		}
+		
+		if (ModuleOutputPin)
+		{
+			// If there's an existing connection, insert our module in the chain
+			if (OutputNodeInputPin->LinkedTo.Num() > 0)
+			{
+				UEdGraphPin* PreviousOutput = OutputNodeInputPin->LinkedTo[0];
+				OutputNodeInputPin->BreakAllPinLinks();
+				
+				// Connect previous to our input
+				UEdGraphPin* ModuleInputPin = nullptr;
+				for (UEdGraphPin* Pin : NewModuleNode->Pins)
+				{
+					if (Pin->Direction == EGPD_Input)
+					{
+						ModuleInputPin = Pin;
+						break;
+					}
+				}
+				if (ModuleInputPin && PreviousOutput)
+				{
+					ModuleInputPin->MakeLinkTo(PreviousOutput);
+				}
+			}
+			
+			// Connect our output to the output node
+			ModuleOutputPin->MakeLinkTo(OutputNodeInputPin);
+		}
+	}
+	
+	// Notify graph changed
+	Graph->NotifyGraphChanged();
+	
+	// Mark as modified and trigger recompile
+	NiagaraSystem->MarkPackageDirty();
+	NiagaraSystem->RequestCompile(false);
+	NiagaraSystem->PostEditChange();
+	
+	TSharedPtr<FJsonObject> Result = MakeResult();
+	Result->SetStringField(TEXT("system_path"), SystemPath);
+	Result->SetStringField(TEXT("emitter_name"), EmitterName);
+	Result->SetStringField(TEXT("module_name"), NewModuleNode->GetFunctionName());
+	Result->SetStringField(TEXT("module_path"), ModulePath);
+	Result->SetStringField(TEXT("script_usage"), ScriptUsageStr);
+	
+	return FECACommandResult::Success(Result);
+}
+
+//------------------------------------------------------------------------------
+// SetNiagaraModuleInput
+//------------------------------------------------------------------------------
+
+FECACommandResult FECACommand_SetNiagaraModuleInput::Execute(const TSharedPtr<FJsonObject>& Params)
+{
+	FString SystemPath;
+	if (!GetStringParam(Params, TEXT("system_path"), SystemPath))
+	{
+		return FECACommandResult::Error(TEXT("Missing required parameter: system_path"));
+	}
+	
+	FString EmitterName;
+	if (!GetStringParam(Params, TEXT("emitter_name"), EmitterName))
+	{
+		return FECACommandResult::Error(TEXT("Missing required parameter: emitter_name"));
+	}
+	
+	FString ModuleName;
+	if (!GetStringParam(Params, TEXT("module_name"), ModuleName))
+	{
+		return FECACommandResult::Error(TEXT("Missing required parameter: module_name"));
+	}
+	
+	FString InputName;
+	if (!GetStringParam(Params, TEXT("input_name"), InputName))
+	{
+		return FECACommandResult::Error(TEXT("Missing required parameter: input_name"));
+	}
+	
+	FString ScriptUsageStr;
+	if (!GetStringParam(Params, TEXT("script_usage"), ScriptUsageStr))
+	{
+		return FECACommandResult::Error(TEXT("Missing required parameter: script_usage"));
+	}
+	
+	// Load the system
+	UNiagaraSystem* NiagaraSystem = LoadObject<UNiagaraSystem>(nullptr, *SystemPath);
+	if (!NiagaraSystem)
+	{
+		return FECACommandResult::Error(FString::Printf(TEXT("Failed to load Niagara system: %s"), *SystemPath));
+	}
+	
+	// Find the emitter
+	FNiagaraEmitterHandle* EmitterHandle = FindEmitterHandleByName(NiagaraSystem, EmitterName);
+	if (!EmitterHandle)
+	{
+		return FECACommandResult::Error(FString::Printf(TEXT("Emitter not found: %s"), *EmitterName));
+	}
+	
+	FVersionedNiagaraEmitterData* EmitterData = EmitterHandle->GetEmitterData();
+	if (!EmitterData)
+	{
+		return FECACommandResult::Error(TEXT("Emitter has no data"));
+	}
+	
+	ENiagaraScriptUsage ScriptUsage = ParseScriptUsage(ScriptUsageStr);
+	
+	UNiagaraScriptSource* ScriptSource = Cast<UNiagaraScriptSource>(EmitterData->GraphSource);
+	if (!ScriptSource || !ScriptSource->NodeGraph)
+	{
+		return FECACommandResult::Error(TEXT("Emitter has no script source"));
+	}
+	
+	UNiagaraGraph* Graph = ScriptSource->NodeGraph;
+	
+	// Find the module node
+	UNiagaraNodeFunctionCall* ModuleNode = nullptr;
+	TArray<UNiagaraNodeFunctionCall*> FunctionNodes;
+	Graph->GetNodesOfClass<UNiagaraNodeFunctionCall>(FunctionNodes);
+	
+	for (UNiagaraNodeFunctionCall* FuncNode : FunctionNodes)
+	{
+		if (FuncNode->GetFunctionName().Equals(ModuleName, ESearchCase::IgnoreCase))
+		{
+			ModuleNode = FuncNode;
+			break;
+		}
+	}
+	
+	if (!ModuleNode)
+	{
+		return FECACommandResult::Error(FString::Printf(TEXT("Module not found: %s"), *ModuleName));
+	}
+	
+	// Find the input pin
+	UEdGraphPin* InputPin = nullptr;
+	for (UEdGraphPin* Pin : ModuleNode->Pins)
+	{
+		if (Pin->Direction == EGPD_Input && Pin->PinName.ToString().Equals(InputName, ESearchCase::IgnoreCase))
+		{
+			InputPin = Pin;
+			break;
+		}
+	}
+	
+	if (!InputPin)
+	{
+		return FECACommandResult::Error(FString::Printf(TEXT("Input not found: %s"), *InputName));
+	}
+	
+	// Set the value based on type
+	// For now, we'll set the default value string which works for simple types
+	double FloatValue;
+	int32 IntValue;
+	bool BoolValue;
+	FVector VectorValue;
+	
+	if (GetFloatParam(Params, TEXT("value"), FloatValue, false))
+	{
+		InputPin->DefaultValue = FString::SanitizeFloat(FloatValue);
+	}
+	else if (GetIntParam(Params, TEXT("value"), IntValue, false))
+	{
+		InputPin->DefaultValue = FString::FromInt(IntValue);
+	}
+	else if (GetBoolParam(Params, TEXT("value"), BoolValue, false))
+	{
+		InputPin->DefaultValue = BoolValue ? TEXT("true") : TEXT("false");
+	}
+	else if (GetVectorParam(Params, TEXT("value"), VectorValue, false))
+	{
+		InputPin->DefaultValue = FString::Printf(TEXT("%f,%f,%f"), VectorValue.X, VectorValue.Y, VectorValue.Z);
+	}
+	else
+	{
+		return FECACommandResult::Error(TEXT("Could not parse value parameter"));
+	}
+	
+	// Mark as modified
+	Graph->NotifyGraphChanged();
+	NiagaraSystem->MarkPackageDirty();
+	NiagaraSystem->RequestCompile(false);
+	NiagaraSystem->PostEditChange();
+	
+	TSharedPtr<FJsonObject> Result = MakeResult();
+	Result->SetStringField(TEXT("module_name"), ModuleName);
+	Result->SetStringField(TEXT("input_name"), InputName);
+	Result->SetStringField(TEXT("value"), InputPin->DefaultValue);
+	
+	return FECACommandResult::Success(Result);
+}
+
+//------------------------------------------------------------------------------
+// GetNiagaraEmitters
+//------------------------------------------------------------------------------
+
+FECACommandResult FECACommand_GetNiagaraEmitters::Execute(const TSharedPtr<FJsonObject>& Params)
+{
+	FString SystemPath;
+	if (!GetStringParam(Params, TEXT("system_path"), SystemPath))
+	{
+		return FECACommandResult::Error(TEXT("Missing required parameter: system_path"));
+	}
+	
+	UNiagaraSystem* NiagaraSystem = LoadObject<UNiagaraSystem>(nullptr, *SystemPath);
+	if (!NiagaraSystem)
+	{
+		return FECACommandResult::Error(FString::Printf(TEXT("Failed to load Niagara system: %s"), *SystemPath));
+	}
+	
+	TArray<TSharedPtr<FJsonValue>> EmittersArray;
+	
+	const TArray<FNiagaraEmitterHandle>& EmitterHandles = NiagaraSystem->GetEmitterHandles();
+	for (int32 i = 0; i < EmitterHandles.Num(); ++i)
+	{
+		const FNiagaraEmitterHandle& Handle = EmitterHandles[i];
+		
+		TSharedPtr<FJsonObject> EmitterJson = MakeShared<FJsonObject>();
+		EmitterJson->SetStringField(TEXT("name"), Handle.GetName().ToString());
+		EmitterJson->SetStringField(TEXT("unique_name"), Handle.GetUniqueInstanceName());
+		EmitterJson->SetStringField(TEXT("id"), Handle.GetId().ToString());
+		EmitterJson->SetNumberField(TEXT("index"), i);
+		EmitterJson->SetBoolField(TEXT("enabled"), Handle.GetIsEnabled());
+		
+		FVersionedNiagaraEmitterData* EmitterData = Handle.GetInstance().GetEmitterData();
+		if (EmitterData)
+		{
+			EmitterJson->SetStringField(TEXT("sim_target"), 
+				EmitterData->SimTarget == ENiagaraSimTarget::CPUSim ? TEXT("cpu") : TEXT("gpu"));
+		}
+		
+		EmittersArray.Add(MakeShared<FJsonValueObject>(EmitterJson));
+	}
+	
+	TSharedPtr<FJsonObject> Result = MakeResult();
+	Result->SetStringField(TEXT("system_path"), SystemPath);
+	Result->SetArrayField(TEXT("emitters"), EmittersArray);
+	Result->SetNumberField(TEXT("count"), EmittersArray.Num());
+	
+	return FECACommandResult::Success(Result);
+}
+
+//------------------------------------------------------------------------------
+// GetNiagaraModules
+//------------------------------------------------------------------------------
+
+FECACommandResult FECACommand_GetNiagaraModules::Execute(const TSharedPtr<FJsonObject>& Params)
+{
+	FString SystemPath;
+	if (!GetStringParam(Params, TEXT("system_path"), SystemPath))
+	{
+		return FECACommandResult::Error(TEXT("Missing required parameter: system_path"));
+	}
+	
+	FString EmitterName;
+	if (!GetStringParam(Params, TEXT("emitter_name"), EmitterName))
+	{
+		return FECACommandResult::Error(TEXT("Missing required parameter: emitter_name"));
+	}
+	
+	FString ScriptUsageFilter = TEXT("all");
+	GetStringParam(Params, TEXT("script_usage"), ScriptUsageFilter, false);
+	
+	UNiagaraSystem* NiagaraSystem = LoadObject<UNiagaraSystem>(nullptr, *SystemPath);
+	if (!NiagaraSystem)
+	{
+		return FECACommandResult::Error(FString::Printf(TEXT("Failed to load Niagara system: %s"), *SystemPath));
+	}
+	
+	FNiagaraEmitterHandle* EmitterHandle = FindEmitterHandleByName(NiagaraSystem, EmitterName);
+	if (!EmitterHandle)
+	{
+		return FECACommandResult::Error(FString::Printf(TEXT("Emitter not found: %s"), *EmitterName));
+	}
+	
+	FVersionedNiagaraEmitterData* EmitterData = EmitterHandle->GetEmitterData();
+	if (!EmitterData)
+	{
+		return FECACommandResult::Error(TEXT("Emitter has no data"));
+	}
+	
+	UNiagaraScriptSource* ScriptSource = Cast<UNiagaraScriptSource>(EmitterData->GraphSource);
+	if (!ScriptSource || !ScriptSource->NodeGraph)
+	{
+		return FECACommandResult::Error(TEXT("Emitter has no script source"));
+	}
+	
+	UNiagaraGraph* Graph = ScriptSource->NodeGraph;
+	
+	TArray<TSharedPtr<FJsonValue>> ModulesArray;
+	
+	// Get all function call nodes (modules)
+	TArray<UNiagaraNodeFunctionCall*> FunctionNodes;
+	Graph->GetNodesOfClass<UNiagaraNodeFunctionCall>(FunctionNodes);
+	
+	for (UNiagaraNodeFunctionCall* FuncNode : FunctionNodes)
+	{
+		// Determine which script usage this module belongs to by finding the connected output node
+		FString ModuleUsage = TEXT("unknown");
+		ENiagaraScriptUsage NodeUsage = FNiagaraStackGraphUtilities::GetOutputNodeUsage(*FuncNode);
+		ModuleUsage = ScriptUsageToString(NodeUsage);
+		
+		// Filter by script usage if specified
+		if (!ScriptUsageFilter.Equals(TEXT("all"), ESearchCase::IgnoreCase) && 
+			!ModuleUsage.Equals(ScriptUsageFilter, ESearchCase::IgnoreCase))
+		{
+			continue;
+		}
+		
+		TSharedPtr<FJsonObject> ModuleJson = MakeShared<FJsonObject>();
+		ModuleJson->SetStringField(TEXT("name"), FuncNode->GetFunctionName());
+		ModuleJson->SetStringField(TEXT("script_usage"), ModuleUsage);
+		
+		if (FuncNode->FunctionScript)
+		{
+			ModuleJson->SetStringField(TEXT("script_path"), FuncNode->FunctionScript->GetPathName());
+		}
+		
+		// Get inputs
+		TArray<TSharedPtr<FJsonValue>> InputsArray;
+		for (UEdGraphPin* Pin : FuncNode->Pins)
+		{
+			if (Pin->Direction == EGPD_Input && !Pin->PinName.IsNone())
+			{
+				TSharedPtr<FJsonObject> InputJson = MakeShared<FJsonObject>();
+				InputJson->SetStringField(TEXT("name"), Pin->PinName.ToString());
+				InputJson->SetStringField(TEXT("type"), Pin->PinType.PinCategory.ToString());
+				InputJson->SetStringField(TEXT("default_value"), Pin->DefaultValue);
+				InputsArray.Add(MakeShared<FJsonValueObject>(InputJson));
+			}
+		}
+		ModuleJson->SetArrayField(TEXT("inputs"), InputsArray);
+		
+		ModulesArray.Add(MakeShared<FJsonValueObject>(ModuleJson));
+	}
+	
+	TSharedPtr<FJsonObject> Result = MakeResult();
+	Result->SetStringField(TEXT("system_path"), SystemPath);
+	Result->SetStringField(TEXT("emitter_name"), EmitterName);
+	Result->SetArrayField(TEXT("modules"), ModulesArray);
+	Result->SetNumberField(TEXT("count"), ModulesArray.Num());
+	
+	return FECACommandResult::Success(Result);
+}
+
+//------------------------------------------------------------------------------
+// ListNiagaraModules
+//------------------------------------------------------------------------------
+
+FECACommandResult FECACommand_ListNiagaraModules::Execute(const TSharedPtr<FJsonObject>& Params)
+{
+	FString Category = TEXT("all");
+	FString Search;
+	GetStringParam(Params, TEXT("category"), Category, false);
+	GetStringParam(Params, TEXT("search"), Search, false);
+	
+	// Common built-in module paths
+	struct FModuleInfo
+	{
+		FString Path;
+		FString Category;
+		FString Name;
+		FString Description;
+	};
+	
+	TArray<FModuleInfo> BuiltInModules = {
+		// Emitter modules
+		{ TEXT("/Niagara/Modules/Emitter/EmitterState.EmitterState"), TEXT("emitter"), TEXT("Emitter State"), TEXT("Controls emitter lifecycle state") },
+		{ TEXT("/Niagara/Modules/Emitter/SpawnRate.SpawnRate"), TEXT("emitter"), TEXT("Spawn Rate"), TEXT("Sets particle spawn rate") },
+		{ TEXT("/Niagara/Modules/Emitter/SpawnBurstInstantaneous.SpawnBurstInstantaneous"), TEXT("emitter"), TEXT("Spawn Burst Instantaneous"), TEXT("Spawns particles in a single burst") },
+		{ TEXT("/Niagara/Modules/Emitter/SpawnPerUnit.SpawnPerUnit"), TEXT("emitter"), TEXT("Spawn Per Unit"), TEXT("Spawns particles based on distance traveled") },
+		
+		// Spawn modules
+		{ TEXT("/Niagara/Modules/Spawn/Location/SystemLocation.SystemLocation"), TEXT("spawn"), TEXT("System Location"), TEXT("Sets particle spawn location to system location") },
+		{ TEXT("/Niagara/Modules/Spawn/Location/SphereLocation.SphereLocation"), TEXT("spawn"), TEXT("Sphere Location"), TEXT("Random location within sphere") },
+		{ TEXT("/Niagara/Modules/Spawn/Location/BoxLocation.BoxLocation"), TEXT("spawn"), TEXT("Box Location"), TEXT("Random location within box") },
+		{ TEXT("/Niagara/Modules/Spawn/Location/CylinderLocation.CylinderLocation"), TEXT("spawn"), TEXT("Cylinder Location"), TEXT("Random location within cylinder") },
+		{ TEXT("/Niagara/Modules/Spawn/Location/TorusLocation.TorusLocation"), TEXT("spawn"), TEXT("Torus Location"), TEXT("Random location on torus surface") },
+		{ TEXT("/Niagara/Modules/Spawn/Location/ConeLocation.ConeLocation"), TEXT("spawn"), TEXT("Cone Location"), TEXT("Random location within cone") },
+		{ TEXT("/Niagara/Modules/Spawn/Location/RingLocation.RingLocation"), TEXT("spawn"), TEXT("Ring Location"), TEXT("Random location on ring") },
+		{ TEXT("/Niagara/Modules/Spawn/Location/DiscLocation.DiscLocation"), TEXT("spawn"), TEXT("Disc Location"), TEXT("Random location on disc") },
+		{ TEXT("/Niagara/Modules/Spawn/Location/GridLocation.GridLocation"), TEXT("spawn"), TEXT("Grid Location"), TEXT("Location on grid") },
+		
+		// Velocity modules
+		{ TEXT("/Niagara/Modules/Spawn/Velocity/AddVelocity.AddVelocity"), TEXT("spawn"), TEXT("Add Velocity"), TEXT("Adds velocity to particles") },
+		{ TEXT("/Niagara/Modules/Spawn/Velocity/AddVelocityInCone.AddVelocityInCone"), TEXT("spawn"), TEXT("Add Velocity In Cone"), TEXT("Random velocity within cone") },
+		{ TEXT("/Niagara/Modules/Spawn/Velocity/AddVelocityFromPoint.AddVelocityFromPoint"), TEXT("spawn"), TEXT("Add Velocity From Point"), TEXT("Velocity away from point") },
+		{ TEXT("/Niagara/Modules/Spawn/Velocity/InheritVelocity.InheritVelocity"), TEXT("spawn"), TEXT("Inherit Velocity"), TEXT("Inherit velocity from parent") },
+		
+		// Initialize modules
+		{ TEXT("/Niagara/Modules/Spawn/Initialize/InitializeParticle.InitializeParticle"), TEXT("spawn"), TEXT("Initialize Particle"), TEXT("Sets initial particle properties") },
+		{ TEXT("/Niagara/Modules/Spawn/Size/InitialSize.InitialSize"), TEXT("spawn"), TEXT("Initial Size"), TEXT("Sets initial particle size") },
+		{ TEXT("/Niagara/Modules/Spawn/Mass/InitialMass.InitialMass"), TEXT("spawn"), TEXT("Initial Mass"), TEXT("Sets initial particle mass") },
+		{ TEXT("/Niagara/Modules/Spawn/Lifetime/SetLifetime.SetLifetime"), TEXT("spawn"), TEXT("Set Lifetime"), TEXT("Sets particle lifetime") },
+		
+		// Update modules
+		{ TEXT("/Niagara/Modules/Update/Lifetime/UpdateAge.UpdateAge"), TEXT("update"), TEXT("Update Age"), TEXT("Updates particle age") },
+		{ TEXT("/Niagara/Modules/Update/Color/Color.Color"), TEXT("update"), TEXT("Color"), TEXT("Sets particle color over lifetime") },
+		{ TEXT("/Niagara/Modules/Update/Color/ScaleColor.ScaleColor"), TEXT("update"), TEXT("Scale Color"), TEXT("Scales particle color") },
+		{ TEXT("/Niagara/Modules/Update/Size/ScaleSize.ScaleSize"), TEXT("update"), TEXT("Scale Size"), TEXT("Scales particle size over lifetime") },
+		{ TEXT("/Niagara/Modules/Update/Size/ScaleSizeBySpeed.ScaleSizeBySpeed"), TEXT("update"), TEXT("Scale Size By Speed"), TEXT("Scales size based on velocity") },
+		
+		// Force modules
+		{ TEXT("/Niagara/Modules/Forces/Gravity.Gravity"), TEXT("forces"), TEXT("Gravity"), TEXT("Applies gravity force") },
+		{ TEXT("/Niagara/Modules/Forces/Drag.Drag"), TEXT("forces"), TEXT("Drag"), TEXT("Applies drag force") },
+		{ TEXT("/Niagara/Modules/Forces/Wind.Wind"), TEXT("forces"), TEXT("Wind"), TEXT("Applies wind force") },
+		{ TEXT("/Niagara/Modules/Forces/PointForce.PointForce"), TEXT("forces"), TEXT("Point Force"), TEXT("Force from a point") },
+		{ TEXT("/Niagara/Modules/Forces/VortexForce.VortexForce"), TEXT("forces"), TEXT("Vortex Force"), TEXT("Rotating vortex force") },
+		{ TEXT("/Niagara/Modules/Forces/CurlNoiseForce.CurlNoiseForce"), TEXT("forces"), TEXT("Curl Noise Force"), TEXT("Curl noise turbulence") },
+		
+		// Solver
+		{ TEXT("/Niagara/Modules/Solvers/SolveForcesAndVelocity.SolveForcesAndVelocity"), TEXT("update"), TEXT("Solve Forces And Velocity"), TEXT("Integrates forces to update position") },
+		
+		// Collision
+		{ TEXT("/Niagara/Modules/Collision/CollisionQuery.CollisionQuery"), TEXT("update"), TEXT("Collision Query"), TEXT("Performs collision detection") },
+		
+		// Rotation
+		{ TEXT("/Niagara/Modules/Spawn/Rotation/InitialMeshOrientation.InitialMeshOrientation"), TEXT("spawn"), TEXT("Initial Mesh Orientation"), TEXT("Sets initial mesh rotation") },
+		{ TEXT("/Niagara/Modules/Update/Rotation/UpdateMeshOrientation.UpdateMeshOrientation"), TEXT("update"), TEXT("Update Mesh Orientation"), TEXT("Updates mesh rotation") },
+		
+		// Sprite facing
+		{ TEXT("/Niagara/Modules/Update/Orientation/SpriteFacingAndAlignment.SpriteFacingAndAlignment"), TEXT("update"), TEXT("Sprite Facing And Alignment"), TEXT("Controls sprite billboard mode") },
+	};
+	
+	TArray<TSharedPtr<FJsonValue>> ModulesArray;
+	
+	for (const FModuleInfo& Module : BuiltInModules)
+	{
+		// Filter by category
+		if (!Category.Equals(TEXT("all"), ESearchCase::IgnoreCase) && 
+			!Module.Category.Equals(Category, ESearchCase::IgnoreCase))
+		{
+			continue;
+		}
+		
+		// Filter by search term
+		if (!Search.IsEmpty() && 
+			!Module.Name.Contains(Search, ESearchCase::IgnoreCase) &&
+			!Module.Path.Contains(Search, ESearchCase::IgnoreCase))
+		{
+			continue;
+		}
+		
+		TSharedPtr<FJsonObject> ModuleJson = MakeShared<FJsonObject>();
+		ModuleJson->SetStringField(TEXT("name"), Module.Name);
+		ModuleJson->SetStringField(TEXT("path"), Module.Path);
+		ModuleJson->SetStringField(TEXT("category"), Module.Category);
+		ModuleJson->SetStringField(TEXT("description"), Module.Description);
+		
+		ModulesArray.Add(MakeShared<FJsonValueObject>(ModuleJson));
+	}
+	
+	TSharedPtr<FJsonObject> Result = MakeResult();
+	Result->SetStringField(TEXT("category_filter"), Category);
+	Result->SetStringField(TEXT("search_filter"), Search);
+	Result->SetArrayField(TEXT("modules"), ModulesArray);
+	Result->SetNumberField(TEXT("count"), ModulesArray.Num());
+	
+	return FECACommandResult::Success(Result);
+}
+
+//------------------------------------------------------------------------------
+// RemoveNiagaraModule
+//------------------------------------------------------------------------------
+
+FECACommandResult FECACommand_RemoveNiagaraModule::Execute(const TSharedPtr<FJsonObject>& Params)
+{
+	FString SystemPath;
+	if (!GetStringParam(Params, TEXT("system_path"), SystemPath))
+	{
+		return FECACommandResult::Error(TEXT("Missing required parameter: system_path"));
+	}
+	
+	FString EmitterName;
+	if (!GetStringParam(Params, TEXT("emitter_name"), EmitterName))
+	{
+		return FECACommandResult::Error(TEXT("Missing required parameter: emitter_name"));
+	}
+	
+	FString ModuleName;
+	if (!GetStringParam(Params, TEXT("module_name"), ModuleName))
+	{
+		return FECACommandResult::Error(TEXT("Missing required parameter: module_name"));
+	}
+	
+	FString ScriptUsageStr;
+	if (!GetStringParam(Params, TEXT("script_usage"), ScriptUsageStr))
+	{
+		return FECACommandResult::Error(TEXT("Missing required parameter: script_usage"));
+	}
+	
+	UNiagaraSystem* NiagaraSystem = LoadObject<UNiagaraSystem>(nullptr, *SystemPath);
+	if (!NiagaraSystem)
+	{
+		return FECACommandResult::Error(FString::Printf(TEXT("Failed to load Niagara system: %s"), *SystemPath));
+	}
+	
+	FNiagaraEmitterHandle* EmitterHandle = FindEmitterHandleByName(NiagaraSystem, EmitterName);
+	if (!EmitterHandle)
+	{
+		return FECACommandResult::Error(FString::Printf(TEXT("Emitter not found: %s"), *EmitterName));
+	}
+	
+	FVersionedNiagaraEmitterData* EmitterData = EmitterHandle->GetEmitterData();
+	if (!EmitterData)
+	{
+		return FECACommandResult::Error(TEXT("Emitter has no data"));
+	}
+	
+	UNiagaraScriptSource* ScriptSource = Cast<UNiagaraScriptSource>(EmitterData->GraphSource);
+	if (!ScriptSource || !ScriptSource->NodeGraph)
+	{
+		return FECACommandResult::Error(TEXT("Emitter has no script source"));
+	}
+	
+	UNiagaraGraph* Graph = ScriptSource->NodeGraph;
+	
+	// Find the module node
+	UNiagaraNodeFunctionCall* ModuleNode = nullptr;
+	TArray<UNiagaraNodeFunctionCall*> FunctionNodes;
+	Graph->GetNodesOfClass<UNiagaraNodeFunctionCall>(FunctionNodes);
+	
+	ENiagaraScriptUsage TargetUsage = ParseScriptUsage(ScriptUsageStr);
+	
+	for (UNiagaraNodeFunctionCall* FuncNode : FunctionNodes)
+	{
+		if (FuncNode->GetFunctionName().Equals(ModuleName, ESearchCase::IgnoreCase))
+		{
+			// Check if it's in the correct script usage
+			ENiagaraScriptUsage NodeUsage = FNiagaraStackGraphUtilities::GetOutputNodeUsage(*FuncNode);
+			if (NodeUsage == TargetUsage)
+			{
+				ModuleNode = FuncNode;
+				break;
+			}
+		}
+	}
+	
+	if (!ModuleNode)
+	{
+		return FECACommandResult::Error(FString::Printf(TEXT("Module not found: %s in %s"), *ModuleName, *ScriptUsageStr));
+	}
+	
+	// Remove the module node
+	Graph->RemoveNode(ModuleNode);
+	Graph->NotifyGraphChanged();
+	
+	// Mark as modified and trigger recompile
+	NiagaraSystem->MarkPackageDirty();
+	NiagaraSystem->RequestCompile(false);
+	NiagaraSystem->PostEditChange();
+	
+	TSharedPtr<FJsonObject> Result = MakeResult();
+	Result->SetStringField(TEXT("removed_module"), ModuleName);
+	Result->SetStringField(TEXT("emitter_name"), EmitterName);
+	Result->SetStringField(TEXT("script_usage"), ScriptUsageStr);
+	
+	return FECACommandResult::Success(Result);
+}
+
+//------------------------------------------------------------------------------
+// DeleteNiagaraEmitter
+//------------------------------------------------------------------------------
+
+FECACommandResult FECACommand_DeleteNiagaraEmitter::Execute(const TSharedPtr<FJsonObject>& Params)
+{
+#if WITH_EDITORONLY_DATA
+	FString SystemPath;
+	if (!GetStringParam(Params, TEXT("system_path"), SystemPath))
+	{
+		return FECACommandResult::Error(TEXT("Missing required parameter: system_path"));
+	}
+	
+	FString EmitterName;
+	if (!GetStringParam(Params, TEXT("emitter_name"), EmitterName))
+	{
+		return FECACommandResult::Error(TEXT("Missing required parameter: emitter_name"));
+	}
+	
+	UNiagaraSystem* NiagaraSystem = LoadObject<UNiagaraSystem>(nullptr, *SystemPath);
+	if (!NiagaraSystem)
+	{
+		return FECACommandResult::Error(FString::Printf(TEXT("Failed to load Niagara system: %s"), *SystemPath));
+	}
+	
+	// Find the emitter to delete
+	int32 EmitterIndex;
+	FNiagaraEmitterHandle* EmitterHandle = FindEmitterHandleByName(NiagaraSystem, EmitterName, &EmitterIndex);
+	if (!EmitterHandle)
+	{
+		return FECACommandResult::Error(FString::Printf(TEXT("Emitter not found: %s"), *EmitterName));
+	}
+	
+	FGuid EmitterId = EmitterHandle->GetId();
+	FString DeletedName = EmitterHandle->GetName().ToString();
+	
+	// Use the system's RemoveEmitterHandle method
+	NiagaraSystem->RemoveEmitterHandle(*EmitterHandle);
+	
+	// Mark the system's graph as changed - NotifyGraphChanged() is exported from UEdGraph base class
+	// RequestCompile() below handles the actual recompilation
+	if (UNiagaraScriptSource* SystemScriptSource = Cast<UNiagaraScriptSource>(NiagaraSystem->GetSystemSpawnScript()->GetLatestSource()))
+	{
+		if (UNiagaraGraph* SystemGraph = SystemScriptSource->NodeGraph)
+		{
+			SystemGraph->NotifyGraphChanged();
+		}
+	}
+	
+	// CRITICAL: Synchronize the overview graph with the system - updates the UI to reflect the removal
+	UNiagaraSystemEditorData* SystemEditorData = Cast<UNiagaraSystemEditorData>(NiagaraSystem->GetEditorData());
+	if (SystemEditorData)
+	{
+		SystemEditorData->SynchronizeOverviewGraphWithSystem(*NiagaraSystem);
+	}
+	
+	// Mark as modified and trigger recompile
+	NiagaraSystem->MarkPackageDirty();
+	NiagaraSystem->RequestCompile(false);
+	NiagaraSystem->PostEditChange();
+	
+	TSharedPtr<FJsonObject> Result = MakeResult();
+	Result->SetStringField(TEXT("deleted_emitter"), DeletedName);
+	Result->SetStringField(TEXT("emitter_id"), EmitterId.ToString());
+	Result->SetNumberField(TEXT("remaining_emitters"), NiagaraSystem->GetEmitterHandles().Num());
+	
+	return FECACommandResult::Success(Result);
+#else
+	return FECACommandResult::Error(TEXT("DeleteNiagaraEmitter is only available in editor builds"));
+#endif
+}
+
+//------------------------------------------------------------------------------
+// SetNiagaraEmitterProperty
+//------------------------------------------------------------------------------
+
+FECACommandResult FECACommand_SetNiagaraEmitterProperty::Execute(const TSharedPtr<FJsonObject>& Params)
+{
+	FString SystemPath;
+	if (!GetStringParam(Params, TEXT("system_path"), SystemPath))
+	{
+		return FECACommandResult::Error(TEXT("Missing required parameter: system_path"));
+	}
+	
+	FString EmitterName;
+	if (!GetStringParam(Params, TEXT("emitter_name"), EmitterName))
+	{
+		return FECACommandResult::Error(TEXT("Missing required parameter: emitter_name"));
+	}
+	
+	FString PropertyName;
+	if (!GetStringParam(Params, TEXT("property"), PropertyName))
+	{
+		return FECACommandResult::Error(TEXT("Missing required parameter: property"));
+	}
+	
+	UNiagaraSystem* NiagaraSystem = LoadObject<UNiagaraSystem>(nullptr, *SystemPath);
+	if (!NiagaraSystem)
+	{
+		return FECACommandResult::Error(FString::Printf(TEXT("Failed to load Niagara system: %s"), *SystemPath));
+	}
+	
+	FNiagaraEmitterHandle* EmitterHandle = FindEmitterHandleByName(NiagaraSystem, EmitterName);
+	if (!EmitterHandle)
+	{
+		return FECACommandResult::Error(FString::Printf(TEXT("Emitter not found: %s"), *EmitterName));
+	}
+	
+	FVersionedNiagaraEmitterData* EmitterData = EmitterHandle->GetEmitterData();
+	if (!EmitterData)
+	{
+		return FECACommandResult::Error(TEXT("Emitter has no data"));
+	}
+	
+	TSharedPtr<FJsonObject> Result = MakeResult();
+	Result->SetStringField(TEXT("emitter_name"), EmitterName);
+	Result->SetStringField(TEXT("property"), PropertyName);
+	
+	if (PropertyName.Equals(TEXT("sim_target"), ESearchCase::IgnoreCase))
+	{
+		FString Value;
+		if (!GetStringParam(Params, TEXT("value"), Value))
+		{
+			return FECACommandResult::Error(TEXT("Missing value for sim_target"));
+		}
+		
+		if (Value.Equals(TEXT("cpu"), ESearchCase::IgnoreCase))
+		{
+			EmitterData->SimTarget = ENiagaraSimTarget::CPUSim;
+		}
+		else if (Value.Equals(TEXT("gpu"), ESearchCase::IgnoreCase))
+		{
+			EmitterData->SimTarget = ENiagaraSimTarget::GPUComputeSim;
+		}
+		else
+		{
+			return FECACommandResult::Error(TEXT("Invalid sim_target value. Use 'cpu' or 'gpu'"));
+		}
+		
+		Result->SetStringField(TEXT("value"), Value);
+	}
+	else if (PropertyName.Equals(TEXT("enabled"), ESearchCase::IgnoreCase))
+	{
+		bool Value;
+		if (!GetBoolParam(Params, TEXT("value"), Value))
+		{
+			return FECACommandResult::Error(TEXT("Missing boolean value for enabled"));
+		}
+		
+		EmitterHandle->SetIsEnabled(Value, *NiagaraSystem, false);
+		Result->SetBoolField(TEXT("value"), Value);
+	}
+	else if (PropertyName.Equals(TEXT("bounds_mode"), ESearchCase::IgnoreCase))
+	{
+		FString Value;
+		if (!GetStringParam(Params, TEXT("value"), Value))
+		{
+			return FECACommandResult::Error(TEXT("Missing string value for bounds_mode"));
+		}
+		
+		if (Value.Equals(TEXT("dynamic"), ESearchCase::IgnoreCase))
+		{
+			EmitterData->CalculateBoundsMode = ENiagaraEmitterCalculateBoundMode::Dynamic;
+		}
+		else if (Value.Equals(TEXT("fixed"), ESearchCase::IgnoreCase))
+		{
+			EmitterData->CalculateBoundsMode = ENiagaraEmitterCalculateBoundMode::Fixed;
+		}
+		else
+		{
+			return FECACommandResult::Error(TEXT("Invalid bounds_mode. Use 'dynamic' or 'fixed'"));
+		}
+		Result->SetStringField(TEXT("value"), Value);
+	}
+	else if (PropertyName.Equals(TEXT("local_space"), ESearchCase::IgnoreCase))
+	{
+		bool Value;
+		if (!GetBoolParam(Params, TEXT("value"), Value))
+		{
+			return FECACommandResult::Error(TEXT("Missing boolean value for local_space"));
+		}
+		
+		EmitterData->bLocalSpace = Value;
+		Result->SetBoolField(TEXT("value"), Value);
+	}
+	else if (PropertyName.Equals(TEXT("determinism"), ESearchCase::IgnoreCase))
+	{
+		bool Value;
+		if (!GetBoolParam(Params, TEXT("value"), Value))
+		{
+			return FECACommandResult::Error(TEXT("Missing boolean value for determinism"));
+		}
+		
+		EmitterData->bDeterminism = Value;
+		Result->SetBoolField(TEXT("value"), Value);
+	}
+	else
+	{
+		return FECACommandResult::Error(FString::Printf(TEXT("Unknown property: %s. Supported: sim_target, enabled, bounds_mode, local_space, determinism"), *PropertyName));
+	}
+	
+	// Mark as modified and trigger recompile
+	NiagaraSystem->MarkPackageDirty();
+	NiagaraSystem->RequestCompile(false);
+	NiagaraSystem->PostEditChange();
+	
+	return FECACommandResult::Success(Result);
+}
+
+
+//------------------------------------------------------------------------------
+// ListNiagaraEmitterTemplates
+//------------------------------------------------------------------------------
+
+FECACommandResult FECACommand_ListNiagaraEmitterTemplates::Execute(const TSharedPtr<FJsonObject>& Params)
+{
+	FString Category = TEXT("all");
+	GetStringParam(Params, TEXT("category"), Category, false);
+	
+	struct FEmitterTemplateInfo
+	{
+		FString Path;
+		FString Name;
+		FString Category;
+		FString Description;
+	};
+	
+	TArray<FEmitterTemplateInfo> Templates = {
+		// Sprite emitters
+		{ TEXT("/Niagara/DefaultAssets/Templates/Emitters/SimpleSpriteBurst"), TEXT("Simple Sprite Burst"), TEXT("sprites"), TEXT("A simple burst of sprites") },
+		{ TEXT("/Niagara/DefaultAssets/Templates/Emitters/Fountain"), TEXT("Fountain"), TEXT("sprites"), TEXT("A continuous fountain of particles") },
+		{ TEXT("/Niagara/DefaultAssets/Templates/Emitters/OmnidirectionalBurst"), TEXT("Omnidirectional Burst"), TEXT("sprites"), TEXT("Particles burst in all directions") },
+		{ TEXT("/Niagara/DefaultAssets/Templates/Emitters/DirectionalBurst"), TEXT("Directional Burst"), TEXT("sprites"), TEXT("Particles burst in a specific direction") },
+		{ TEXT("/Niagara/DefaultAssets/Templates/Emitters/ConfettiBurst"), TEXT("Confetti Burst"), TEXT("sprites"), TEXT("Confetti-style particle burst") },
+		{ TEXT("/Niagara/DefaultAssets/Templates/Emitters/HangingParticulates"), TEXT("Hanging Particulates"), TEXT("sprites"), TEXT("Floating dust-like particles") },
+		{ TEXT("/Niagara/DefaultAssets/Templates/Emitters/BlowingParticles"), TEXT("Blowing Particles"), TEXT("sprites"), TEXT("Particles blown by wind") },
+		{ TEXT("/Niagara/DefaultAssets/Templates/Emitters/SingleLoopingParticle"), TEXT("Single Looping Particle"), TEXT("sprites"), TEXT("A single particle that loops") },
+		{ TEXT("/Niagara/DefaultAssets/Templates/Emitters/Minimal"), TEXT("Minimal"), TEXT("sprites"), TEXT("Minimal emitter setup for customization") },
+		{ TEXT("/Niagara/DefaultAssets/Templates/Emitters/RecycleParticlesInView"), TEXT("Recycle Particles In View"), TEXT("sprites"), TEXT("Particles recycled when in camera view") },
+		
+		// Mesh emitters
+		{ TEXT("/Niagara/DefaultAssets/Templates/Emitters/UpwardMeshBurst"), TEXT("Upward Mesh Burst"), TEXT("meshes"), TEXT("Mesh particles burst upward") },
+		
+		// Ribbon emitters
+		{ TEXT("/Niagara/DefaultAssets/Templates/Emitters/LocationBasedRibbon"), TEXT("Location Based Ribbon"), TEXT("ribbons"), TEXT("Ribbon that follows a path") },
+		
+		// Beam emitters
+		{ TEXT("/Niagara/DefaultAssets/Templates/Emitters/StaticBeam"), TEXT("Static Beam"), TEXT("beams"), TEXT("A static beam effect") },
+		{ TEXT("/Niagara/DefaultAssets/Templates/Emitters/DynamicBeam"), TEXT("Dynamic Beam"), TEXT("beams"), TEXT("A dynamic, animated beam") },
+		
+		// Behavior examples
+		{ TEXT("/Niagara/DefaultAssets/Templates/BehaviorExamples/GridLocation"), TEXT("Grid Location"), TEXT("behaviors"), TEXT("Particles spawned in a grid pattern") },
+		{ TEXT("/Niagara/DefaultAssets/Templates/BehaviorExamples/MeshOrientation"), TEXT("Mesh Orientation"), TEXT("behaviors"), TEXT("Example of mesh particle orientation") },
+		{ TEXT("/Niagara/DefaultAssets/Templates/BehaviorExamples/MeshRotationForce"), TEXT("Mesh Rotation Force"), TEXT("behaviors"), TEXT("Mesh particles with rotation") },
+		{ TEXT("/Niagara/DefaultAssets/Templates/BehaviorExamples/SpriteFacingAndAlignment"), TEXT("Sprite Facing And Alignment"), TEXT("behaviors"), TEXT("Sprite billboard and alignment options") },
+		{ TEXT("/Niagara/DefaultAssets/Templates/BehaviorExamples/SubUVAnimation"), TEXT("SubUV Animation"), TEXT("behaviors"), TEXT("Sprite sheet animation") },
+		{ TEXT("/Niagara/DefaultAssets/Templates/BehaviorExamples/RibbonID"), TEXT("Ribbon ID"), TEXT("behaviors"), TEXT("Multiple ribbon trails example") },
+		{ TEXT("/Niagara/DefaultAssets/Templates/BehaviorExamples/RibbonShapes"), TEXT("Ribbon Shapes"), TEXT("behaviors"), TEXT("Different ribbon shape modes") },
+		{ TEXT("/Niagara/DefaultAssets/Templates/BehaviorExamples/PlayAudio"), TEXT("Play Audio"), TEXT("behaviors"), TEXT("Audio playback with particles") },
+		{ TEXT("/Niagara/DefaultAssets/Templates/BehaviorExamples/KillParticles"), TEXT("Kill Particles"), TEXT("behaviors"), TEXT("Particle death conditions") },
+		{ TEXT("/Niagara/DefaultAssets/Templates/BehaviorExamples/SpawnGroups"), TEXT("Spawn Groups"), TEXT("behaviors"), TEXT("Grouped particle spawning") },
+	};
+	
+	TArray<TSharedPtr<FJsonValue>> TemplatesArray;
+	
+	for (const FEmitterTemplateInfo& Template : Templates)
+	{
+		// Apply category filter
+		if (!Category.Equals(TEXT("all"), ESearchCase::IgnoreCase) &&
+			!Template.Category.Equals(Category, ESearchCase::IgnoreCase))
+		{
+			continue;
+		}
+		
+		TSharedPtr<FJsonObject> TemplateJson = MakeShared<FJsonObject>();
+		TemplateJson->SetStringField(TEXT("path"), Template.Path);
+		TemplateJson->SetStringField(TEXT("name"), Template.Name);
+		TemplateJson->SetStringField(TEXT("category"), Template.Category);
+		TemplateJson->SetStringField(TEXT("description"), Template.Description);
+		
+		TemplatesArray.Add(MakeShared<FJsonValueObject>(TemplateJson));
+	}
+	
+	TSharedPtr<FJsonObject> Result = MakeResult();
+	Result->SetStringField(TEXT("category_filter"), Category);
+	Result->SetArrayField(TEXT("templates"), TemplatesArray);
+	Result->SetNumberField(TEXT("count"), TemplatesArray.Num());
+	
+	return FECACommandResult::Success(Result);
+}
+
+//------------------------------------------------------------------------------
+// CreateNiagaraEmitter - Create a new emitter asset with full configuration
+//------------------------------------------------------------------------------
+
+FECACommandResult FECACommand_CreateNiagaraEmitter::Execute(const TSharedPtr<FJsonObject>& Params)
+{
+#if WITH_EDITORONLY_DATA
+	// Get required parameters
+	FString AssetPath;
+	if (!GetStringParam(Params, TEXT("asset_path"), AssetPath))
+	{
+		return FECACommandResult::Error(TEXT("Missing required parameter: asset_path"));
+	}
+	
+	// Get optional parameters with defaults
+	FString Template = TEXT("sprite");
+	GetStringParam(Params, TEXT("template"), Template, false);
+	
+	FString SimTargetStr = TEXT("cpu");
+	GetStringParam(Params, TEXT("sim_target"), SimTargetStr, false);
+	
+	bool bLocalSpace = false;
+	GetBoolParam(Params, TEXT("local_space"), bLocalSpace, false);
+	
+	bool bDeterminism = false;
+	GetBoolParam(Params, TEXT("determinism"), bDeterminism, false);
+	
+	FString BoundsMode = TEXT("dynamic");
+	GetStringParam(Params, TEXT("bounds_mode"), BoundsMode, false);
+	
+	// Parse asset path
+	FString PackagePath;
+	FString AssetName;
+	
+	int32 LastSlash;
+	if (AssetPath.FindLastChar('/', LastSlash))
+	{
+		PackagePath = AssetPath.Left(LastSlash);
+		AssetName = AssetPath.Mid(LastSlash + 1);
+	}
+	else
+	{
+		return FECACommandResult::Error(TEXT("Invalid asset path format. Expected format: /Game/Path/AssetName"));
+	}
+	
+	// Determine the base template emitter to use
+	FString TemplateEmitterPath;
+	if (Template.StartsWith(TEXT("/")))
+	{
+		// User specified a full path to an emitter
+		TemplateEmitterPath = Template;
+	}
+	else if (Template.Equals(TEXT("empty"), ESearchCase::IgnoreCase))
+	{
+		// Minimal emitter
+		TemplateEmitterPath = TEXT("/Niagara/DefaultAssets/Templates/Emitters/Minimal.Minimal");
+	}
+	else if (Template.Equals(TEXT("mesh"), ESearchCase::IgnoreCase))
+	{
+		TemplateEmitterPath = TEXT("/Niagara/DefaultAssets/Templates/Emitters/UpwardMeshBurst.UpwardMeshBurst");
+	}
+	else if (Template.Equals(TEXT("ribbon"), ESearchCase::IgnoreCase))
+	{
+		TemplateEmitterPath = TEXT("/Niagara/DefaultAssets/Templates/Emitters/LocationBasedRibbon.LocationBasedRibbon");
+	}
+	else // sprite (default)
+	{
+		TemplateEmitterPath = TEXT("/Niagara/DefaultAssets/Templates/Emitters/Fountain.Fountain");
+	}
+	
+	// Load the template emitter
+	UNiagaraEmitter* TemplateEmitter = LoadObject<UNiagaraEmitter>(nullptr, *TemplateEmitterPath);
+	if (!TemplateEmitter)
+	{
+		return FECACommandResult::Error(FString::Printf(TEXT("Failed to load template emitter: %s"), *TemplateEmitterPath));
+	}
+	
+	// Create the package
+	FString FullPackagePath = PackagePath / AssetName;
+	UPackage* Package = CreatePackage(*FullPackagePath);
+	if (!Package)
+	{
+		return FECACommandResult::Error(FString::Printf(TEXT("Failed to create package at: %s"), *FullPackagePath));
+	}
+	Package->FullyLoad();
+	
+	// Duplicate the template emitter
+	UNiagaraEmitter* NewEmitter = Cast<UNiagaraEmitter>(StaticDuplicateObject(TemplateEmitter, Package, FName(*AssetName), RF_Public | RF_Standalone));
+	if (!NewEmitter)
+	{
+		return FECACommandResult::Error(TEXT("Failed to duplicate template emitter"));
+	}
+	
+	NewEmitter->SetUniqueEmitterName(AssetName);
+	NewEmitter->DisableVersioning(TemplateEmitter->GetExposedVersion().VersionGuid);
+	
+	// Configure the emitter properties
+	FVersionedNiagaraEmitterData* EmitterData = NewEmitter->GetLatestEmitterData();
+	if (EmitterData)
+	{
+		// Set simulation target
+		if (SimTargetStr.Equals(TEXT("gpu"), ESearchCase::IgnoreCase))
+		{
+			EmitterData->SimTarget = ENiagaraSimTarget::GPUComputeSim;
+		}
+		else
+		{
+			EmitterData->SimTarget = ENiagaraSimTarget::CPUSim;
+		}
+		
+		// Set emitter properties
+		EmitterData->bLocalSpace = bLocalSpace;
+		EmitterData->bDeterminism = bDeterminism;
+		
+		if (BoundsMode.Equals(TEXT("fixed"), ESearchCase::IgnoreCase))
+		{
+			EmitterData->CalculateBoundsMode = ENiagaraEmitterCalculateBoundMode::Fixed;
+		}
+		else
+		{
+			EmitterData->CalculateBoundsMode = ENiagaraEmitterCalculateBoundMode::Dynamic;
+		}
+	}
+	
+	// Make the emitter inheritable by default
+	NewEmitter->bIsInheritable = true;
+	NewEmitter->TemplateAssetDescription = FText();
+	NewEmitter->Category = FText();
+	
+	// Mark package dirty and save
+	Package->MarkPackageDirty();
+	
+	// Notify asset registry
+	FAssetRegistryModule::AssetCreated(NewEmitter);
+	
+	// Save the asset
+	FString PackageFileName = FPackageName::LongPackageNameToFilename(FullPackagePath, FPackageName::GetAssetPackageExtension());
+	FSavePackageArgs SaveArgs;
+	SaveArgs.TopLevelFlags = RF_Public | RF_Standalone;
+	UPackage::SavePackage(Package, NewEmitter, *PackageFileName, SaveArgs);
+	
+	// Build result
+	TSharedPtr<FJsonObject> Result = MakeResult();
+	Result->SetStringField(TEXT("asset_path"), NewEmitter->GetPathName());
+	Result->SetStringField(TEXT("asset_name"), NewEmitter->GetName());
+	Result->SetStringField(TEXT("package_path"), PackagePath);
+	Result->SetStringField(TEXT("template"), Template);
+	Result->SetStringField(TEXT("template_source"), TemplateEmitterPath);
+	Result->SetStringField(TEXT("sim_target"), SimTargetStr);
+	Result->SetBoolField(TEXT("local_space"), bLocalSpace);
+	Result->SetBoolField(TEXT("determinism"), bDeterminism);
+	Result->SetStringField(TEXT("bounds_mode"), BoundsMode);
+	
+	return FECACommandResult::Success(Result);
+#else
+	return FECACommandResult::Error(TEXT("CreateNiagaraEmitter is only available in editor builds"));
+#endif
+}
+
+//------------------------------------------------------------------------------
+// AddNiagaraRenderer - Add a renderer to a Niagara emitter
+//------------------------------------------------------------------------------
+
+FECACommandResult FECACommand_AddNiagaraRenderer::Execute(const TSharedPtr<FJsonObject>& Params)
+{
+#if WITH_EDITORONLY_DATA
+	FString SystemPath;
+	if (!GetStringParam(Params, TEXT("system_path"), SystemPath))
+	{
+		return FECACommandResult::Error(TEXT("Missing required parameter: system_path"));
+	}
+	
+	FString EmitterName;
+	if (!GetStringParam(Params, TEXT("emitter_name"), EmitterName))
+	{
+		return FECACommandResult::Error(TEXT("Missing required parameter: emitter_name"));
+	}
+	
+	FString RendererType;
+	if (!GetStringParam(Params, TEXT("renderer_type"), RendererType))
+	{
+		return FECACommandResult::Error(TEXT("Missing required parameter: renderer_type"));
+	}
+	
+	FString RendererName;
+	GetStringParam(Params, TEXT("renderer_name"), RendererName, false);
+	
+	FString MaterialPath;
+	GetStringParam(Params, TEXT("material_path"), MaterialPath, false);
+	
+	FString MeshPath;
+	GetStringParam(Params, TEXT("mesh_path"), MeshPath, false);
+	
+	// Load the system
+	UNiagaraSystem* NiagaraSystem = LoadObject<UNiagaraSystem>(nullptr, *SystemPath);
+	if (!NiagaraSystem)
+	{
+		return FECACommandResult::Error(FString::Printf(TEXT("Failed to load Niagara system: %s"), *SystemPath));
+	}
+	
+	// Find the emitter
+	FNiagaraEmitterHandle* EmitterHandle = FindEmitterHandleByName(NiagaraSystem, EmitterName);
+	if (!EmitterHandle)
+	{
+		return FECACommandResult::Error(FString::Printf(TEXT("Emitter not found: %s"), *EmitterName));
+	}
+	
+	FVersionedNiagaraEmitterData* EmitterData = EmitterHandle->GetEmitterData();
+	if (!EmitterData)
+	{
+		return FECACommandResult::Error(TEXT("Emitter has no data"));
+	}
+	
+	UNiagaraEmitter* Emitter = EmitterHandle->GetInstance().Emitter;
+	if (!Emitter)
+	{
+		return FECACommandResult::Error(TEXT("Failed to get emitter instance"));
+	}
+	
+	// Create the renderer based on type
+	UNiagaraRendererProperties* NewRenderer = nullptr;
+	FString ActualRendererName = RendererName.IsEmpty() ? RendererType + TEXT("Renderer") : RendererName;
+	
+	if (RendererType.Equals(TEXT("sprite"), ESearchCase::IgnoreCase))
+	{
+		UNiagaraSpriteRendererProperties* SpriteRenderer = NewObject<UNiagaraSpriteRendererProperties>(Emitter, FName(*ActualRendererName));
+		
+		// Load and assign material if specified
+		if (!MaterialPath.IsEmpty())
+		{
+			UMaterialInterface* Material = LoadObject<UMaterialInterface>(nullptr, *MaterialPath);
+			if (Material)
+			{
+				SpriteRenderer->Material = Material;
+			}
+		}
+		
+		NewRenderer = SpriteRenderer;
+	}
+	else if (RendererType.Equals(TEXT("mesh"), ESearchCase::IgnoreCase))
+	{
+		UNiagaraMeshRendererProperties* MeshRenderer = NewObject<UNiagaraMeshRendererProperties>(Emitter, FName(*ActualRendererName));
+		
+		// Load and assign mesh if specified
+		if (!MeshPath.IsEmpty())
+		{
+			UStaticMesh* Mesh = LoadObject<UStaticMesh>(nullptr, *MeshPath);
+			if (Mesh)
+			{
+				FNiagaraMeshRendererMeshProperties MeshProps;
+				MeshProps.Mesh = Mesh;
+				MeshRenderer->Meshes.Add(MeshProps);
+			}
+		}
+		
+		// Load and assign material if specified
+		if (!MaterialPath.IsEmpty())
+		{
+			UMaterialInterface* Material = LoadObject<UMaterialInterface>(nullptr, *MaterialPath);
+			if (Material && MeshRenderer->Meshes.Num() > 0)
+			{
+				FNiagaraMeshMaterialOverride Override;
+				Override.ExplicitMat = Material;
+				MeshRenderer->OverrideMaterials.Add(Override);
+			}
+		}
+		
+		NewRenderer = MeshRenderer;
+	}
+	else if (RendererType.Equals(TEXT("ribbon"), ESearchCase::IgnoreCase))
+	{
+		UNiagaraRibbonRendererProperties* RibbonRenderer = NewObject<UNiagaraRibbonRendererProperties>(Emitter, FName(*ActualRendererName));
+		
+		if (!MaterialPath.IsEmpty())
+		{
+			UMaterialInterface* Material = LoadObject<UMaterialInterface>(nullptr, *MaterialPath);
+			if (Material)
+			{
+				RibbonRenderer->Material = Material;
+			}
+		}
+		
+		NewRenderer = RibbonRenderer;
+	}
+	else if (RendererType.Equals(TEXT("light"), ESearchCase::IgnoreCase))
+	{
+		NewRenderer = NewObject<UNiagaraLightRendererProperties>(Emitter, FName(*ActualRendererName));
+	}
+	else if (RendererType.Equals(TEXT("component"), ESearchCase::IgnoreCase))
+	{
+		NewRenderer = NewObject<UNiagaraComponentRendererProperties>(Emitter, FName(*ActualRendererName));
+	}
+	else
+	{
+		return FECACommandResult::Error(FString::Printf(TEXT("Unknown renderer type: %s. Supported: sprite, mesh, ribbon, light, component"), *RendererType));
+	}
+	
+	// Add the renderer to the emitter
+	Emitter->AddRenderer(NewRenderer, EmitterData->Version.VersionGuid);
+	
+	// Mark as modified
+	NiagaraSystem->MarkPackageDirty();
+	NiagaraSystem->RequestCompile(false);
+	NiagaraSystem->PostEditChange();
+	
+	TSharedPtr<FJsonObject> Result = MakeResult();
+	Result->SetStringField(TEXT("renderer_name"), ActualRendererName);
+	Result->SetStringField(TEXT("renderer_type"), RendererType);
+	Result->SetStringField(TEXT("emitter_name"), EmitterName);
+	Result->SetNumberField(TEXT("renderer_count"), EmitterData->GetRenderers().Num());
+	
+	return FECACommandResult::Success(Result);
+#else
+	return FECACommandResult::Error(TEXT("AddNiagaraRenderer is only available in editor builds"));
+#endif
+}
+
+//------------------------------------------------------------------------------
+// SetNiagaraMaterial - Set material on a Niagara renderer
+//------------------------------------------------------------------------------
+
+FECACommandResult FECACommand_SetNiagaraMaterial::Execute(const TSharedPtr<FJsonObject>& Params)
+{
+#if WITH_EDITORONLY_DATA
+	FString SystemPath;
+	if (!GetStringParam(Params, TEXT("system_path"), SystemPath))
+	{
+		return FECACommandResult::Error(TEXT("Missing required parameter: system_path"));
+	}
+	
+	FString EmitterName;
+	if (!GetStringParam(Params, TEXT("emitter_name"), EmitterName))
+	{
+		return FECACommandResult::Error(TEXT("Missing required parameter: emitter_name"));
+	}
+	
+	FString MaterialPath;
+	if (!GetStringParam(Params, TEXT("material_path"), MaterialPath))
+	{
+		return FECACommandResult::Error(TEXT("Missing required parameter: material_path"));
+	}
+	
+	int32 RendererIndex = 0;
+	GetIntParam(Params, TEXT("renderer_index"), RendererIndex, false);
+	
+	// Load the system
+	UNiagaraSystem* NiagaraSystem = LoadObject<UNiagaraSystem>(nullptr, *SystemPath);
+	if (!NiagaraSystem)
+	{
+		return FECACommandResult::Error(FString::Printf(TEXT("Failed to load Niagara system: %s"), *SystemPath));
+	}
+	
+	// Find the emitter
+	FNiagaraEmitterHandle* EmitterHandle = FindEmitterHandleByName(NiagaraSystem, EmitterName);
+	if (!EmitterHandle)
+	{
+		return FECACommandResult::Error(FString::Printf(TEXT("Emitter not found: %s"), *EmitterName));
+	}
+	
+	FVersionedNiagaraEmitterData* EmitterData = EmitterHandle->GetEmitterData();
+	if (!EmitterData)
+	{
+		return FECACommandResult::Error(TEXT("Emitter has no data"));
+	}
+	
+	// Load the material
+	UMaterialInterface* Material = LoadObject<UMaterialInterface>(nullptr, *MaterialPath);
+	if (!Material)
+	{
+		return FECACommandResult::Error(FString::Printf(TEXT("Failed to load material: %s"), *MaterialPath));
+	}
+	
+	// Get the renderers
+	const TArray<UNiagaraRendererProperties*>& Renderers = EmitterData->GetRenderers();
+	if (RendererIndex < 0 || RendererIndex >= Renderers.Num())
+	{
+		return FECACommandResult::Error(FString::Printf(TEXT("Renderer index %d out of range. Emitter has %d renderers."), RendererIndex, Renderers.Num()));
+	}
+	
+	UNiagaraRendererProperties* Renderer = Renderers[RendererIndex];
+	FString RendererTypeName;
+	
+	// Set material based on renderer type
+	if (UNiagaraSpriteRendererProperties* SpriteRenderer = Cast<UNiagaraSpriteRendererProperties>(Renderer))
+	{
+		SpriteRenderer->Material = Material;
+		RendererTypeName = TEXT("sprite");
+	}
+	else if (UNiagaraMeshRendererProperties* MeshRenderer = Cast<UNiagaraMeshRendererProperties>(Renderer))
+	{
+		FNiagaraMeshMaterialOverride Override;
+		Override.ExplicitMat = Material;
+		if (MeshRenderer->OverrideMaterials.Num() > 0)
+		{
+			MeshRenderer->OverrideMaterials[0] = Override;
+		}
+		else
+		{
+			MeshRenderer->OverrideMaterials.Add(Override);
+		}
+		RendererTypeName = TEXT("mesh");
+	}
+	else if (UNiagaraRibbonRendererProperties* RibbonRenderer = Cast<UNiagaraRibbonRendererProperties>(Renderer))
+	{
+		RibbonRenderer->Material = Material;
+		RendererTypeName = TEXT("ribbon");
+	}
+	else
+	{
+		return FECACommandResult::Error(TEXT("Renderer type does not support materials"));
+	}
+	
+	// Mark as modified
+	NiagaraSystem->MarkPackageDirty();
+	NiagaraSystem->RequestCompile(false);
+	NiagaraSystem->PostEditChange();
+	
+	TSharedPtr<FJsonObject> Result = MakeResult();
+	Result->SetStringField(TEXT("material_path"), MaterialPath);
+	Result->SetStringField(TEXT("renderer_type"), RendererTypeName);
+	Result->SetNumberField(TEXT("renderer_index"), RendererIndex);
+	Result->SetStringField(TEXT("emitter_name"), EmitterName);
+	
+	return FECACommandResult::Success(Result);
+#else
+	return FECACommandResult::Error(TEXT("SetNiagaraMaterial is only available in editor builds"));
+#endif
+}
+
+//------------------------------------------------------------------------------
+// SetNiagaraCurve - Set a curve on a Niagara module input
+//------------------------------------------------------------------------------
+
+FECACommandResult FECACommand_SetNiagaraCurve::Execute(const TSharedPtr<FJsonObject>& Params)
+{
+#if WITH_EDITORONLY_DATA
+	FString SystemPath;
+	if (!GetStringParam(Params, TEXT("system_path"), SystemPath))
+	{
+		return FECACommandResult::Error(TEXT("Missing required parameter: system_path"));
+	}
+	
+	FString EmitterName;
+	if (!GetStringParam(Params, TEXT("emitter_name"), EmitterName))
+	{
+		return FECACommandResult::Error(TEXT("Missing required parameter: emitter_name"));
+	}
+	
+	FString ModuleName;
+	if (!GetStringParam(Params, TEXT("module_name"), ModuleName))
+	{
+		return FECACommandResult::Error(TEXT("Missing required parameter: module_name"));
+	}
+	
+	FString InputName;
+	if (!GetStringParam(Params, TEXT("input_name"), InputName))
+	{
+		return FECACommandResult::Error(TEXT("Missing required parameter: input_name"));
+	}
+	
+	FString ScriptUsageStr;
+	if (!GetStringParam(Params, TEXT("script_usage"), ScriptUsageStr))
+	{
+		return FECACommandResult::Error(TEXT("Missing required parameter: script_usage"));
+	}
+	
+	FString CurveType;
+	if (!GetStringParam(Params, TEXT("curve_type"), CurveType))
+	{
+		return FECACommandResult::Error(TEXT("Missing required parameter: curve_type"));
+	}
+	
+	const TArray<TSharedPtr<FJsonValue>>* KeysArray;
+	if (!Params->TryGetArrayField(TEXT("keys"), KeysArray))
+	{
+		return FECACommandResult::Error(TEXT("Missing required parameter: keys"));
+	}
+	
+	// Load the system
+	UNiagaraSystem* NiagaraSystem = LoadObject<UNiagaraSystem>(nullptr, *SystemPath);
+	if (!NiagaraSystem)
+	{
+		return FECACommandResult::Error(FString::Printf(TEXT("Failed to load Niagara system: %s"), *SystemPath));
+	}
+	
+	// Find the emitter
+	FNiagaraEmitterHandle* EmitterHandle = FindEmitterHandleByName(NiagaraSystem, EmitterName);
+	if (!EmitterHandle)
+	{
+		return FECACommandResult::Error(FString::Printf(TEXT("Emitter not found: %s"), *EmitterName));
+	}
+	
+	// Check if this is a Stateless emitter - they use Distribution system instead of DataInterface curves
+	// reinterpret_cast needed because UNiagaraStatelessEmitter is forward declared (internal header)
+	// We know it inherits from UObject, but the compiler can't verify with incomplete type
+	UObject* StatelessEmitter = reinterpret_cast<UObject*>(EmitterHandle->GetStatelessEmitter());
+	ENiagaraEmitterMode EmitterMode = EmitterHandle->GetEmitterMode();
+	
+	// Debug: Log emitter mode info
+	UE_LOG(LogTemp, Log, TEXT("SetNiagaraCurve: Emitter '%s' Mode=%d, StatelessEmitter=%p"), 
+		*EmitterName, (int32)EmitterMode, StatelessEmitter);
+	
+	if (StatelessEmitter && EmitterMode == ENiagaraEmitterMode::Stateless)
+	{
+		// Handle Stateless emitter via reflection - find the Modules array property
+		FArrayProperty* ModulesArrayProp = CastField<FArrayProperty>(StatelessEmitter->GetClass()->FindPropertyByName(TEXT("Modules")));
+		if (!ModulesArrayProp)
+		{
+			return FECACommandResult::Error(TEXT("Could not find Modules property on Stateless emitter"));
+		}
+		
+		FScriptArrayHelper ModulesArray(ModulesArrayProp, ModulesArrayProp->ContainerPtrToValuePtr<void>(StatelessEmitter));
+		UObject* TargetModule = nullptr;
+		
+		for (int32 i = 0; i < ModulesArray.Num(); ++i)
+		{
+			UObject** ModulePtr = reinterpret_cast<UObject**>(ModulesArray.GetRawPtr(i));
+			if (ModulePtr && *ModulePtr)
+			{
+				UObject* Module = *ModulePtr;
+				if (Module->GetClass()->GetName().Contains(ModuleName, ESearchCase::IgnoreCase))
+				{
+					TargetModule = Module;
+					break;
+				}
+			}
+		}
+		
+		if (!TargetModule)
+		{
+			return FECACommandResult::Error(FString::Printf(TEXT("Stateless module not found: %s"), *ModuleName));
+		}
+		
+		// Find the Distribution property by name
+		// Stateless modules use properties like "ScaleDistribution", "ColorDistribution", etc.
+		FProperty* DistributionProp = nullptr;
+		FString FoundPropName;
+		for (TFieldIterator<FProperty> PropIt(TargetModule->GetClass()); PropIt; ++PropIt)
+		{
+			FProperty* Prop = *PropIt;
+			FString PropName = Prop->GetName();
+			
+			// Check if this is a Distribution property (struct property ending in "Distribution")
+			FStructProperty* StructProp = CastField<FStructProperty>(Prop);
+			if (StructProp && PropName.Contains(TEXT("Distribution"), ESearchCase::IgnoreCase))
+			{
+				// For ScaleColor module, the property is "ScaleDistribution"
+				// Match if input name contains part of the property name or vice versa
+				// e.g., InputName="ColorCurve" or "Scale" should match "ScaleDistribution"
+				FString PropNameWithoutDist = PropName.Replace(TEXT("Distribution"), TEXT(""));
+				if (InputName.Contains(PropNameWithoutDist, ESearchCase::IgnoreCase) ||
+				    PropNameWithoutDist.Contains(InputName, ESearchCase::IgnoreCase) ||
+				    InputName.Contains(TEXT("Color"), ESearchCase::IgnoreCase) ||
+				    InputName.Contains(TEXT("Scale"), ESearchCase::IgnoreCase) ||
+				    InputName.Contains(TEXT("Curve"), ESearchCase::IgnoreCase))
+				{
+					DistributionProp = Prop;
+					FoundPropName = PropName;
+					break;
+				}
+			}
+		}
+		
+		if (!DistributionProp)
+		{
+			// List available properties for debugging
+			FString AvailableProps;
+			for (TFieldIterator<FProperty> PropIt(TargetModule->GetClass()); PropIt; ++PropIt)
+			{
+				if (!AvailableProps.IsEmpty()) AvailableProps += TEXT(", ");
+				AvailableProps += (*PropIt)->GetName();
+			}
+			return FECACommandResult::Error(FString::Printf(TEXT("Distribution property not found for input '%s'. Available properties: %s"), *InputName, *AvailableProps));
+		}
+		
+		// Handle color distribution (FNiagaraDistributionColor)
+		if (CurveType.Equals(TEXT("color"), ESearchCase::IgnoreCase))
+		{
+			FStructProperty* StructProp = CastField<FStructProperty>(DistributionProp);
+			if (StructProp && StructProp->Struct->GetFName() == TEXT("NiagaraDistributionColor"))
+			{
+				FNiagaraDistributionColor* ColorDist = StructProp->ContainerPtrToValuePtr<FNiagaraDistributionColor>(TargetModule);
+				if (ColorDist)
+				{
+					// Set mode to curve - this tells the system to use curve-based interpolation
+					ColorDist->Mode = ENiagaraDistributionMode::UniformCurve;
+					
+					// For FNiagaraDistributionColor, we need to set up the ChannelCurves array
+					// which is inherited from FNiagaraDistributionBase, then call UpdateValuesFromDistribution()
+					// to convert the curve data into the Values array format
+					
+					// Access the base class ChannelCurves - ensure we have 4 curves (R, G, B, A)
+					FNiagaraDistributionBase* BaseDistribution = static_cast<FNiagaraDistributionBase*>(ColorDist);
+					BaseDistribution->ChannelCurves.SetNum(4);
+					for (int32 i = 0; i < 4; ++i)
+					{
+						BaseDistribution->ChannelCurves[i].Reset();
+					}
+					
+					int32 KeyCount = 0;
+					for (const TSharedPtr<FJsonValue>& KeyValue : *KeysArray)
+					{
+						const TSharedPtr<FJsonObject>* KeyObj;
+						if (KeyValue->TryGetObject(KeyObj))
+						{
+							float Time = static_cast<float>((*KeyObj)->GetNumberField(TEXT("time")));
+							float R = static_cast<float>((*KeyObj)->GetNumberField(TEXT("r")));
+							float G = static_cast<float>((*KeyObj)->GetNumberField(TEXT("g")));
+							float B = static_cast<float>((*KeyObj)->GetNumberField(TEXT("b")));
+							float A = (*KeyObj)->HasField(TEXT("a")) ? static_cast<float>((*KeyObj)->GetNumberField(TEXT("a"))) : 1.0f;
+							
+							BaseDistribution->ChannelCurves[0].AddKey(Time, R);
+							BaseDistribution->ChannelCurves[1].AddKey(Time, G);
+							BaseDistribution->ChannelCurves[2].AddKey(Time, B);
+							BaseDistribution->ChannelCurves[3].AddKey(Time, A);
+							KeyCount++;
+						}
+					}
+					
+					// Set the time range for the curve (0 to 1 for normalized age)
+					ColorDist->ValuesTimeRange = FVector2f(0.0f, 1.0f);
+					
+					// CRITICAL: Update the distribution from the curves
+					// This converts ChannelCurves data into the Values array format
+					ColorDist->UpdateValuesFromDistribution();
+					
+					// Mark modified
+					TargetModule->Modify();
+					StatelessEmitter->Modify();
+					NiagaraSystem->MarkPackageDirty();
+					NiagaraSystem->PostEditChange();
+					
+					// CRITICAL: Force recompile and wait for it to complete
+					// This ensures our Distribution changes are compiled into the shader
+					NiagaraSystem->RequestCompile(true);
+					NiagaraSystem->WaitForCompilationComplete();
+					
+					TSharedPtr<FJsonObject> Result = MakeResult();
+					Result->SetStringField(TEXT("module_name"), ModuleName);
+					Result->SetStringField(TEXT("input_name"), InputName);
+					Result->SetStringField(TEXT("property_name"), FoundPropName);
+					Result->SetStringField(TEXT("curve_type"), CurveType);
+					Result->SetNumberField(TEXT("key_count"), KeyCount);
+					Result->SetBoolField(TEXT("curve_created"), true);
+					Result->SetBoolField(TEXT("stateless_emitter"), true);
+					
+					return FECACommandResult::Success(Result);
+				}
+			}
+		}
+		
+		return FECACommandResult::Error(FString::Printf(TEXT("Unsupported curve type for Stateless emitter: %s"), *CurveType));
+	}
+	
+	// Standard emitter path - use Dynamic Input with ColorCurve
+	FVersionedNiagaraEmitterData* EmitterData = EmitterHandle->GetEmitterData();
+	if (!EmitterData)
+	{
+		return FECACommandResult::Error(TEXT("Emitter has no data"));
+	}
+	
+	// Get the script source and graph
+	UNiagaraScriptSource* ScriptSource = Cast<UNiagaraScriptSource>(EmitterData->GraphSource);
+	if (!ScriptSource || !ScriptSource->NodeGraph)
+	{
+		return FECACommandResult::Error(TEXT("Emitter has no script source"));
+	}
+	
+	UNiagaraGraph* Graph = ScriptSource->NodeGraph;
+	
+	// Find the module node
+	TArray<UNiagaraNodeFunctionCall*> FunctionNodes;
+	Graph->GetNodesOfClass<UNiagaraNodeFunctionCall>(FunctionNodes);
+	
+	UNiagaraNodeFunctionCall* ModuleNode = nullptr;
+	ENiagaraScriptUsage TargetUsage = ParseScriptUsage(ScriptUsageStr);
+	
+	for (UNiagaraNodeFunctionCall* FuncNode : FunctionNodes)
+	{
+		if (FuncNode->GetFunctionName().Equals(ModuleName, ESearchCase::IgnoreCase))
+		{
+			ENiagaraScriptUsage NodeUsage = FNiagaraStackGraphUtilities::GetOutputNodeUsage(*FuncNode);
+			if (NodeUsage == TargetUsage)
+			{
+				ModuleNode = FuncNode;
+				break;
+			}
+		}
+	}
+	
+	if (!ModuleNode)
+	{
+		return FECACommandResult::Error(FString::Printf(TEXT("Module not found: %s in %s"), *ModuleName, *ScriptUsageStr));
+	}
+	
+	// Debug: Log available inputs on this module
+	UE_LOG(LogTemp, Log, TEXT("SetNiagaraCurve: Found module '%s', listing inputs:"), *ModuleNode->GetFunctionName());
+	
+	// Find the target input pin and check if it's a data interface type
+	UEdGraphPin* TargetInputPin = nullptr;
+	bool bIsDataInterfaceInput = false;
+	
+	for (UEdGraphPin* Pin : ModuleNode->Pins)
+	{
+		if (Pin->Direction == EGPD_Input)
+		{
+			UE_LOG(LogTemp, Log, TEXT("  Input: '%s' Type: '%s' DefaultValue: '%s'"), 
+				*Pin->PinName.ToString(), *Pin->PinType.PinCategory.ToString(), *Pin->DefaultValue);
+			
+			// Check if this is our target input
+			if (Pin->PinName.ToString().Equals(InputName, ESearchCase::IgnoreCase))
+			{
+				TargetInputPin = Pin;
+				// "Class" type means it's a data interface input
+				// The pin category for data interfaces is "class" (as shown in debug logs)
+				bIsDataInterfaceInput = Pin->PinType.PinCategory.ToString().Equals(TEXT("Class"), ESearchCase::IgnoreCase);
+			}
+		}
+	}
+	
+	// ============================================================================
+	// APPROACH SELECTION:
+	// - If target input is a DataInterface type (e.g., ScaleColor.ColorCurve), 
+	//   set the curve data interface directly on that input.
+	// - If target input is a value type (e.g., Color.Color which expects a LinearColor),
+	//   use Dynamic Input (ColorFromCurve) to sample from a curve.
+	// ============================================================================
+	
+	if (bIsDataInterfaceInput)
+	{
+		UE_LOG(LogTemp, Log, TEXT("SetNiagaraCurve: Input '%s' is a DataInterface type - setting curve directly"), *InputName);
+		
+		// Determine the curve data interface class based on curve type
+		UClass* CurveClass = nullptr;
+		if (CurveType.Equals(TEXT("float"), ESearchCase::IgnoreCase))
+		{
+			CurveClass = UNiagaraDataInterfaceCurve::StaticClass();
+		}
+		else if (CurveType.Equals(TEXT("vector"), ESearchCase::IgnoreCase))
+		{
+			CurveClass = UNiagaraDataInterfaceVectorCurve::StaticClass();
+		}
+		else if (CurveType.Equals(TEXT("color"), ESearchCase::IgnoreCase))
+		{
+			CurveClass = UNiagaraDataInterfaceColorCurve::StaticClass();
+		}
+		else
+		{
+			return FECACommandResult::Error(FString::Printf(TEXT("Unknown curve type: %s"), *CurveType));
+		}
+		
+		// Create the parameter handle for the input
+		FNiagaraParameterHandle InputHandle = FNiagaraParameterHandle::CreateModuleParameterHandle(FName(*InputName));
+		FNiagaraParameterHandle AliasedHandle = FNiagaraParameterHandle::CreateAliasedModuleParameterHandle(InputHandle, ModuleNode);
+		
+		// Get the type definition for the curve
+		FNiagaraTypeDefinition CurveTypeDef(CurveClass);
+		
+		// Get or create the override pin
+		UEdGraphPin& OverridePin = FNiagaraStackGraphUtilities::GetOrCreateStackFunctionInputOverridePin(
+			*ModuleNode, AliasedHandle, CurveTypeDef, FGuid(), FGuid());
+		
+		UE_LOG(LogTemp, Log, TEXT("SetNiagaraCurve: Override pin: '%s', LinkedTo: %d"), 
+			*OverridePin.PinName.ToString(), OverridePin.LinkedTo.Num());
+		
+		// Remove any existing override
+		if (OverridePin.LinkedTo.Num() > 0)
+		{
+			RemoveOverridePinConnections(OverridePin, Graph);
+		}
+		
+		// Create the curve data interface directly on this input
+		UNiagaraDataInterface* CurveDataInterface = nullptr;
+		FString AliasedInputName = AliasedHandle.GetParameterHandleString().ToString();
+		FNiagaraStackGraphUtilities::SetDataInterfaceValueForFunctionInput(
+			OverridePin, CurveClass, AliasedInputName, CurveDataInterface, FGuid());
+		
+		if (!CurveDataInterface)
+		{
+			return FECACommandResult::Error(TEXT("Failed to create curve data interface"));
+		}
+		
+		UE_LOG(LogTemp, Log, TEXT("SetNiagaraCurve: Created curve data interface: %s at %p"), 
+			*CurveDataInterface->GetClass()->GetName(), CurveDataInterface);
+		
+		// Use a scoped transaction
+		FScopedTransaction Transaction(NSLOCTEXT("ECABridge", "SetNiagaraCurveDirect", "Set Niagara Curve Direct"));
+		CurveDataInterface->Modify();
+		
+		// Track number of keys added
+		int32 KeyCount = 0;
+		
+		// Populate the curve data interface with the key data
+		if (CurveType.Equals(TEXT("color"), ESearchCase::IgnoreCase))
+		{
+			UNiagaraDataInterfaceColorCurve* ColorCurve = Cast<UNiagaraDataInterfaceColorCurve>(CurveDataInterface);
+			if (ColorCurve)
+			{
+				ColorCurve->CurveAsset = nullptr;
+				ColorCurve->RedCurve.Reset();
+				ColorCurve->GreenCurve.Reset();
+				ColorCurve->BlueCurve.Reset();
+				ColorCurve->AlphaCurve.Reset();
+				
+				for (const TSharedPtr<FJsonValue>& KeyValue : *KeysArray)
+				{
+					const TSharedPtr<FJsonObject>* KeyObj;
+					if (KeyValue->TryGetObject(KeyObj))
+					{
+						float Time = static_cast<float>((*KeyObj)->GetNumberField(TEXT("time")));
+						float R = static_cast<float>((*KeyObj)->GetNumberField(TEXT("r")));
+						float G = static_cast<float>((*KeyObj)->GetNumberField(TEXT("g")));
+						float B = static_cast<float>((*KeyObj)->GetNumberField(TEXT("b")));
+						float A = (*KeyObj)->HasField(TEXT("a")) ? static_cast<float>((*KeyObj)->GetNumberField(TEXT("a"))) : 1.0f;
+						
+						ColorCurve->RedCurve.AddKey(Time, R);
+						ColorCurve->GreenCurve.AddKey(Time, G);
+						ColorCurve->BlueCurve.AddKey(Time, B);
+						ColorCurve->AlphaCurve.AddKey(Time, A);
+						KeyCount++;
+					}
+				}
+				
+				ColorCurve->UpdateTimeRanges();
+				ColorCurve->UpdateLUT();
+				ColorCurve->MarkRenderDataDirty();
+				ColorCurve->OnChanged().Broadcast();
+				
+				// Debug: Check the LUT time range
+				UE_LOG(LogTemp, Log, TEXT("SetNiagaraCurve: Populated ColorCurve with %d keys (direct DI)"), KeyCount);
+				UE_LOG(LogTemp, Log, TEXT("SetNiagaraCurve: LUT TimeRange [%.3f, %.3f]"), 
+					ColorCurve->GetMinTime(), ColorCurve->GetMaxTime());
+				// Log the first key values to verify they were set
+				if (ColorCurve->RedCurve.GetNumKeys() > 0)
+				{
+					FKeyHandle FirstKey = ColorCurve->RedCurve.GetFirstKeyHandle();
+					float R = ColorCurve->RedCurve.GetKeyValue(FirstKey);
+					float G = ColorCurve->GreenCurve.GetKeyValue(ColorCurve->GreenCurve.GetFirstKeyHandle());
+					float B = ColorCurve->BlueCurve.GetKeyValue(ColorCurve->BlueCurve.GetFirstKeyHandle());
+					UE_LOG(LogTemp, Log, TEXT("SetNiagaraCurve: First curve key values: R=%.3f G=%.3f B=%.3f"), R, G, B);
+				}
+			}
+		}
+		else if (CurveType.Equals(TEXT("float"), ESearchCase::IgnoreCase))
+		{
+			UNiagaraDataInterfaceCurve* FloatCurve = Cast<UNiagaraDataInterfaceCurve>(CurveDataInterface);
+			if (FloatCurve)
+			{
+				FloatCurve->CurveAsset = nullptr;
+				FloatCurve->Curve.Reset();
+				
+				for (const TSharedPtr<FJsonValue>& KeyValue : *KeysArray)
+				{
+					const TSharedPtr<FJsonObject>* KeyObj;
+					if (KeyValue->TryGetObject(KeyObj))
+					{
+						float Time = static_cast<float>((*KeyObj)->GetNumberField(TEXT("time")));
+						float Value = static_cast<float>((*KeyObj)->GetNumberField(TEXT("value")));
+						FloatCurve->Curve.AddKey(Time, Value);
+						KeyCount++;
+					}
+				}
+				
+				FloatCurve->UpdateTimeRanges();
+				FloatCurve->UpdateLUT();
+				FloatCurve->OnChanged().Broadcast();
+				
+				UE_LOG(LogTemp, Log, TEXT("SetNiagaraCurve: Populated FloatCurve with %d keys (direct DI)"), KeyCount);
+			}
+		}
+		else if (CurveType.Equals(TEXT("vector"), ESearchCase::IgnoreCase))
+		{
+			UNiagaraDataInterfaceVectorCurve* VectorCurve = Cast<UNiagaraDataInterfaceVectorCurve>(CurveDataInterface);
+			if (VectorCurve)
+			{
+				VectorCurve->CurveAsset = nullptr;
+				VectorCurve->XCurve.Reset();
+				VectorCurve->YCurve.Reset();
+				VectorCurve->ZCurve.Reset();
+				
+				for (const TSharedPtr<FJsonValue>& KeyValue : *KeysArray)
+				{
+					const TSharedPtr<FJsonObject>* KeyObj;
+					if (KeyValue->TryGetObject(KeyObj))
+					{
+						float Time = static_cast<float>((*KeyObj)->GetNumberField(TEXT("time")));
+						float X = static_cast<float>((*KeyObj)->GetNumberField(TEXT("x")));
+						float Y = static_cast<float>((*KeyObj)->GetNumberField(TEXT("y")));
+						float Z = static_cast<float>((*KeyObj)->GetNumberField(TEXT("z")));
+						VectorCurve->XCurve.AddKey(Time, X);
+						VectorCurve->YCurve.AddKey(Time, Y);
+						VectorCurve->ZCurve.AddKey(Time, Z);
+						KeyCount++;
+					}
+				}
+				
+				VectorCurve->UpdateTimeRanges();
+				VectorCurve->UpdateLUT();
+				VectorCurve->OnChanged().Broadcast();
+				
+				UE_LOG(LogTemp, Log, TEXT("SetNiagaraCurve: Populated VectorCurve with %d keys (direct DI)"), KeyCount);
+			}
+		}
+		
+		// Find the InputNode that owns this data interface and explicitly notify it
+		// Note: GetDataInterface/SetDataInterface are not exported, use reflection
+		for (UEdGraphNode* Node : Graph->Nodes)
+		{
+			if (UNiagaraNodeInput* InputNode = Cast<UNiagaraNodeInput>(Node))
+			{
+				// Use reflection to access DataInterface property
+				FObjectProperty* DIProperty = CastField<FObjectProperty>(InputNode->GetClass()->FindPropertyByName(TEXT("DataInterface")));
+				if (DIProperty)
+				{
+					UNiagaraDataInterface* NodeDI = Cast<UNiagaraDataInterface>(DIProperty->GetObjectPropertyValue_InContainer(InputNode));
+					if (NodeDI == CurveDataInterface)
+					{
+						// Mark the node as needing synchronization
+						InputNode->MarkNodeRequiresSynchronization(TEXT("Curve data interface modified via ECA"), true);
+						UE_LOG(LogTemp, Log, TEXT("SetNiagaraCurve: Found and marked InputNode for synchronization"));
+						break;
+					}
+				}
+			}
+		}
+		
+		// Mark everything as modified and trigger recompile
+		// Note: NotifyGraphNeedsRecompile and NotifyGraphDataInterfaceChanged are not exported
+		// NotifyGraphChanged() is exported from UEdGraph base class and handles the update
+		Graph->NotifyGraphChanged();
+		ModuleNode->MarkNodeRequiresSynchronization(TEXT("Curve set via ECA"), true);
+		
+		// Note: UNiagaraScriptSource::MarkNotSynchronized is not exported, so we rely on
+		// Graph->NotifyGraphChanged() and ModuleNode->MarkNodeRequiresSynchronization() above
+		// to trigger the necessary recompilation
+		
+		NiagaraSystem->MarkPackageDirty();
+		NiagaraSystem->PostEditChange();
+		NiagaraSystem->RequestCompile(true);
+		NiagaraSystem->WaitForCompilationComplete();
+		
+		UE_LOG(LogTemp, Log, TEXT("SetNiagaraCurve: Compilation complete (direct DI approach)"));
+		
+		// Auto-set Scale Mode if needed
+		for (UEdGraphPin* Pin : ModuleNode->Pins)
+		{
+			if (Pin->Direction == EGPD_Input && Pin->PinName.ToString().Contains(TEXT("Scale Mode")))
+			{
+				if (Pin->PinType.PinSubCategoryObject.IsValid())
+				{
+					UEnum* EnumType = Cast<UEnum>(Pin->PinType.PinSubCategoryObject.Get());
+					if (EnumType)
+					{
+						FString BestModeValue;
+						for (int32 i = 0; i < EnumType->NumEnums() - 1; ++i)
+						{
+							FString EnumName = EnumType->GetNameStringByIndex(i);
+							FString DisplayName = EnumType->GetDisplayNameTextByIndex(i).ToString();
+							
+							bool bIsCurveMode = DisplayName.Contains(TEXT("Curve"), ESearchCase::IgnoreCase);
+							bool bIsConstantMode = !bIsCurveMode && (DisplayName.Contains(TEXT("Linear Color"), ESearchCase::IgnoreCase) || 
+								DisplayName.Contains(TEXT("Constant"), ESearchCase::IgnoreCase) ||
+								DisplayName.Contains(TEXT("Separately"), ESearchCase::IgnoreCase) ||
+								DisplayName.Contains(TEXT("Together"), ESearchCase::IgnoreCase));
+							
+							if (bIsCurveMode)
+							{
+								BestModeValue = EnumName;
+								UE_LOG(LogTemp, Log, TEXT("SetNiagaraCurve: Found curve Scale Mode: '%s'"), *EnumName);
+								break;
+							}
+						}
+						
+						if (!BestModeValue.IsEmpty() && !Pin->DefaultValue.Equals(BestModeValue))
+						{
+							UE_LOG(LogTemp, Log, TEXT("SetNiagaraCurve: Setting Scale Mode from '%s' to '%s'"), 
+								*Pin->DefaultValue, *BestModeValue);
+							Pin->DefaultValue = BestModeValue;
+							
+							NiagaraSystem->RequestCompile(true);
+							NiagaraSystem->WaitForCompilationComplete();
+							UE_LOG(LogTemp, Log, TEXT("SetNiagaraCurve: Recompiled after Scale Mode change"));
+						}
+					}
+				}
+			}
+		}
+		
+		// Return success
+		TSharedPtr<FJsonObject> Result = MakeResult();
+		Result->SetStringField(TEXT("module_name"), ModuleName);
+		Result->SetStringField(TEXT("input_name"), InputName);
+		Result->SetStringField(TEXT("curve_type"), CurveType);
+		Result->SetNumberField(TEXT("key_count"), KeyCount);
+		Result->SetBoolField(TEXT("curve_created"), true);
+		Result->SetBoolField(TEXT("dynamic_input_used"), false);
+		Result->SetBoolField(TEXT("direct_data_interface"), true);
+		Result->SetBoolField(TEXT("stateless_emitter"), false);
+		return FECACommandResult::Success(Result);
+	}
+	
+	// ============================================================================
+	// DYNAMIC INPUT APPROACH: For value-type inputs (Color, Float, Vector)
+	// Use Dynamic Input (ColorFromCurve, FloatFromCurve, VectorFromCurve)
+	// to sample from a curve and provide the value.
+	// ============================================================================
+	
+	UE_LOG(LogTemp, Log, TEXT("SetNiagaraCurve: Input '%s' is a value type - using Dynamic Input approach"), *InputName);
+	
+	// Determine the Dynamic Input script path based on curve type
+	FString DynamicInputPath;
+	FNiagaraTypeDefinition InputType;
+	if (CurveType.Equals(TEXT("float"), ESearchCase::IgnoreCase))
+	{
+		DynamicInputPath = TEXT("/Niagara/DynamicInputs/ValueFromCurve/FloatFromCurve.FloatFromCurve");
+		InputType = FNiagaraTypeDefinition::GetFloatDef();
+	}
+	else if (CurveType.Equals(TEXT("vector"), ESearchCase::IgnoreCase))
+	{
+		DynamicInputPath = TEXT("/Niagara/DynamicInputs/ValueFromCurve/VectorFromCurve.VectorFromCurve");
+		InputType = FNiagaraTypeDefinition::GetVec3Def();
+	}
+	else if (CurveType.Equals(TEXT("color"), ESearchCase::IgnoreCase))
+	{
+		DynamicInputPath = TEXT("/Niagara/DynamicInputs/ValueFromCurve/ColorFromCurve.ColorFromCurve");
+		InputType = FNiagaraTypeDefinition::GetColorDef();
+	}
+	else
+	{
+		return FECACommandResult::Error(FString::Printf(TEXT("Unknown curve type: %s. Supported: float, vector, color"), *CurveType));
+	}
+	
+	// Load the Dynamic Input script
+	UNiagaraScript* DynamicInputScript = LoadObject<UNiagaraScript>(nullptr, *DynamicInputPath);
+	if (!DynamicInputScript)
+	{
+		return FECACommandResult::Error(FString::Printf(TEXT("Failed to load Dynamic Input script: %s"), *DynamicInputPath));
+	}
+	
+	UE_LOG(LogTemp, Log, TEXT("SetNiagaraCurve: Loaded Dynamic Input script: %s"), *DynamicInputPath);
+	
+	// Create the parameter handle for the input
+	FNiagaraParameterHandle InputHandle = FNiagaraParameterHandle::CreateModuleParameterHandle(FName(*InputName));
+	FNiagaraParameterHandle AliasedHandle = FNiagaraParameterHandle::CreateAliasedModuleParameterHandle(InputHandle, ModuleNode);
+	
+	// Get or create the override pin for the module input
+	UEdGraphPin& OverridePin = FNiagaraStackGraphUtilities::GetOrCreateStackFunctionInputOverridePin(
+		*ModuleNode, AliasedHandle, InputType, FGuid(), FGuid());
+	
+	UE_LOG(LogTemp, Log, TEXT("SetNiagaraCurve: Override pin: '%s', LinkedTo: %d"), 
+		*OverridePin.PinName.ToString(), OverridePin.LinkedTo.Num());
+	
+	// Remove any existing override
+	if (OverridePin.LinkedTo.Num() > 0)
+	{
+		UE_LOG(LogTemp, Log, TEXT("SetNiagaraCurve: Removing existing override..."));
+		RemoveOverridePinConnections(OverridePin, Graph);
+	}
+	
+	// Set the Dynamic Input on the module's input pin
+	UNiagaraNodeFunctionCall* DynamicInputNode = nullptr;
+	FNiagaraStackGraphUtilities::SetDynamicInputForFunctionInput(
+		OverridePin, DynamicInputScript, DynamicInputNode, FGuid(), 
+		CurveType + TEXT("FromCurve"), FGuid());
+	
+	if (!DynamicInputNode)
+	{
+		return FECACommandResult::Error(TEXT("Failed to create Dynamic Input node"));
+	}
+	
+	UE_LOG(LogTemp, Log, TEXT("SetNiagaraCurve: Created Dynamic Input node: %s"), *DynamicInputNode->GetFunctionName());
+	
+	// Now we need to find the ColorCurve/FloatCurve/VectorCurve input on the Dynamic Input node
+	// and set our curve data on it
+	UE_LOG(LogTemp, Log, TEXT("SetNiagaraCurve: Listing Dynamic Input node inputs:"));
+	for (UEdGraphPin* Pin : DynamicInputNode->Pins)
+	{
+		if (Pin->Direction == EGPD_Input)
+		{
+			UE_LOG(LogTemp, Log, TEXT("  DI Input: '%s' Type: '%s'"), 
+				*Pin->PinName.ToString(), *Pin->PinType.PinCategory.ToString());
+		}
+	}
+	
+	// Find the curve input on the Dynamic Input
+	// The input name is "DefaultCurve" for all FromCurve dynamic inputs (ColorFromCurve, FloatFromCurve, VectorFromCurve)
+	// The type of curve is determined by the data interface class, not the input name
+	FString CurveInputName = TEXT("DefaultCurve");
+	
+	// Create the parameter handle for the curve input on the Dynamic Input
+	FNiagaraParameterHandle CurveInputHandle = FNiagaraParameterHandle::CreateModuleParameterHandle(FName(*CurveInputName));
+	FNiagaraParameterHandle AliasedCurveHandle = FNiagaraParameterHandle::CreateAliasedModuleParameterHandle(CurveInputHandle, DynamicInputNode);
+	
+	// Determine the curve data interface type
+	UClass* CurveClass = nullptr;
+	FNiagaraTypeDefinition CurveInputType;
+	if (CurveType.Equals(TEXT("float"), ESearchCase::IgnoreCase))
+	{
+		CurveClass = UNiagaraDataInterfaceCurve::StaticClass();
+		CurveInputType = FNiagaraTypeDefinition(UNiagaraDataInterfaceCurve::StaticClass());
+	}
+	else if (CurveType.Equals(TEXT("vector"), ESearchCase::IgnoreCase))
+	{
+		CurveClass = UNiagaraDataInterfaceVectorCurve::StaticClass();
+		CurveInputType = FNiagaraTypeDefinition(UNiagaraDataInterfaceVectorCurve::StaticClass());
+	}
+	else if (CurveType.Equals(TEXT("color"), ESearchCase::IgnoreCase))
+	{
+		CurveClass = UNiagaraDataInterfaceColorCurve::StaticClass();
+		CurveInputType = FNiagaraTypeDefinition(UNiagaraDataInterfaceColorCurve::StaticClass());
+	}
+	
+	// Get or create the override pin for the curve input on the Dynamic Input
+	UEdGraphPin& CurveOverridePin = FNiagaraStackGraphUtilities::GetOrCreateStackFunctionInputOverridePin(
+		*DynamicInputNode, AliasedCurveHandle, CurveInputType, FGuid(), FGuid());
+	
+	UE_LOG(LogTemp, Log, TEXT("SetNiagaraCurve: Curve override pin: '%s', LinkedTo: %d"), 
+		*CurveOverridePin.PinName.ToString(), CurveOverridePin.LinkedTo.Num());
+	
+	// Remove any existing curve override
+	if (CurveOverridePin.LinkedTo.Num() > 0)
+	{
+		RemoveOverridePinConnections(CurveOverridePin, Graph);
+	}
+	
+	// Create the curve data interface
+	UNiagaraDataInterface* CurveDataInterface = nullptr;
+	FString AliasedCurveInputName = AliasedCurveHandle.GetParameterHandleString().ToString();
+	FNiagaraStackGraphUtilities::SetDataInterfaceValueForFunctionInput(
+		CurveOverridePin, CurveClass, AliasedCurveInputName, CurveDataInterface, FGuid());
+	
+	if (!CurveDataInterface)
+	{
+		return FECACommandResult::Error(TEXT("Failed to create curve data interface"));
+	}
+	
+	UE_LOG(LogTemp, Log, TEXT("SetNiagaraCurve: Created curve data interface: %s at %p"), 
+		*CurveDataInterface->GetClass()->GetName(), CurveDataInterface);
+	
+	// CRITICAL: Use a scoped transaction
+	FScopedTransaction Transaction(NSLOCTEXT("ECABridge", "SetNiagaraCurve", "Set Niagara Curve"));
+	CurveDataInterface->Modify();
+	
+	// Track number of keys added
+	int32 KeyCount = 0;
+	
+	// Now populate the curve data interface with the key data
+	if (CurveType.Equals(TEXT("color"), ESearchCase::IgnoreCase))
+	{
+		UNiagaraDataInterfaceColorCurve* ColorCurve = Cast<UNiagaraDataInterfaceColorCurve>(CurveDataInterface);
+		if (ColorCurve)
+		{
+			// Clear CurveAsset to prevent SyncCurvesToAsset from overwriting our data
+			ColorCurve->CurveAsset = nullptr;
+			
+			// Clear existing keys
+			ColorCurve->RedCurve.Reset();
+			ColorCurve->GreenCurve.Reset();
+			ColorCurve->BlueCurve.Reset();
+			ColorCurve->AlphaCurve.Reset();
+			
+			// Add keys from the input array
+			for (const TSharedPtr<FJsonValue>& KeyValue : *KeysArray)
+			{
+				const TSharedPtr<FJsonObject>* KeyObj;
+				if (KeyValue->TryGetObject(KeyObj))
+				{
+					float Time = static_cast<float>((*KeyObj)->GetNumberField(TEXT("time")));
+					float R = static_cast<float>((*KeyObj)->GetNumberField(TEXT("r")));
+					float G = static_cast<float>((*KeyObj)->GetNumberField(TEXT("g")));
+					float B = static_cast<float>((*KeyObj)->GetNumberField(TEXT("b")));
+					float A = (*KeyObj)->HasField(TEXT("a")) ? static_cast<float>((*KeyObj)->GetNumberField(TEXT("a"))) : 1.0f;
+					
+					ColorCurve->RedCurve.AddKey(Time, R);
+					ColorCurve->GreenCurve.AddKey(Time, G);
+					ColorCurve->BlueCurve.AddKey(Time, B);
+					ColorCurve->AlphaCurve.AddKey(Time, A);
+					KeyCount++;
+				}
+			}
+			
+			// Update the curve's internal data
+			ColorCurve->UpdateTimeRanges();
+			ColorCurve->UpdateLUT();
+			ColorCurve->OnChanged().Broadcast();
+			
+			UE_LOG(LogTemp, Log, TEXT("SetNiagaraCurve: Populated ColorCurve with %d keys"), KeyCount);
+		}
+	}
+	else if (CurveType.Equals(TEXT("float"), ESearchCase::IgnoreCase))
+	{
+		UNiagaraDataInterfaceCurve* FloatCurve = Cast<UNiagaraDataInterfaceCurve>(CurveDataInterface);
+		if (FloatCurve)
+		{
+			FloatCurve->CurveAsset = nullptr;
+			FloatCurve->Curve.Reset();
+			
+			for (const TSharedPtr<FJsonValue>& KeyValue : *KeysArray)
+			{
+				const TSharedPtr<FJsonObject>* KeyObj;
+				if (KeyValue->TryGetObject(KeyObj))
+				{
+					float Time = static_cast<float>((*KeyObj)->GetNumberField(TEXT("time")));
+					float Value = static_cast<float>((*KeyObj)->GetNumberField(TEXT("value")));
+					FloatCurve->Curve.AddKey(Time, Value);
+					KeyCount++;
+				}
+			}
+			
+			FloatCurve->UpdateTimeRanges();
+			FloatCurve->UpdateLUT();
+			FloatCurve->OnChanged().Broadcast();
+			
+			UE_LOG(LogTemp, Log, TEXT("SetNiagaraCurve: Populated FloatCurve with %d keys"), KeyCount);
+		}
+	}
+	else if (CurveType.Equals(TEXT("vector"), ESearchCase::IgnoreCase))
+	{
+		UNiagaraDataInterfaceVectorCurve* VectorCurve = Cast<UNiagaraDataInterfaceVectorCurve>(CurveDataInterface);
+		if (VectorCurve)
+		{
+			VectorCurve->CurveAsset = nullptr;
+			VectorCurve->XCurve.Reset();
+			VectorCurve->YCurve.Reset();
+			VectorCurve->ZCurve.Reset();
+			
+			for (const TSharedPtr<FJsonValue>& KeyValue : *KeysArray)
+			{
+				const TSharedPtr<FJsonObject>* KeyObj;
+				if (KeyValue->TryGetObject(KeyObj))
+				{
+					float Time = static_cast<float>((*KeyObj)->GetNumberField(TEXT("time")));
+					float X = static_cast<float>((*KeyObj)->GetNumberField(TEXT("x")));
+					float Y = static_cast<float>((*KeyObj)->GetNumberField(TEXT("y")));
+					float Z = static_cast<float>((*KeyObj)->GetNumberField(TEXT("z")));
+					VectorCurve->XCurve.AddKey(Time, X);
+					VectorCurve->YCurve.AddKey(Time, Y);
+					VectorCurve->ZCurve.AddKey(Time, Z);
+					KeyCount++;
+				}
+			}
+			
+			VectorCurve->UpdateTimeRanges();
+			VectorCurve->UpdateLUT();
+			VectorCurve->OnChanged().Broadcast();
+			
+			UE_LOG(LogTemp, Log, TEXT("SetNiagaraCurve: Populated VectorCurve with %d keys"), KeyCount);
+		}
+	}
+	
+	// Mark everything as modified and trigger recompile
+	Graph->NotifyGraphChanged();
+	DynamicInputNode->MarkNodeRequiresSynchronization(TEXT("Curve set via ECA"), true);
+	ModuleNode->MarkNodeRequiresSynchronization(TEXT("Input overridden via ECA"), true);
+	
+	NiagaraSystem->MarkPackageDirty();
+	NiagaraSystem->PostEditChange();
+	NiagaraSystem->RequestCompile(true);
+	NiagaraSystem->WaitForCompilationComplete();
+	
+	UE_LOG(LogTemp, Log, TEXT("SetNiagaraCurve: Compilation complete"));
+	
+	// Auto-set Scale Mode for modules that require it (like ScaleColor)
+	// The Scale Mode enum determines whether the module uses constant values or curves
+	for (UEdGraphPin* Pin : ModuleNode->Pins)
+	{
+		if (Pin->Direction == EGPD_Input && Pin->PinName.ToString().Contains(TEXT("Scale Mode")))
+		{
+			// Check if the pin has an enum type
+			if (Pin->PinType.PinSubCategoryObject.IsValid())
+			{
+				UEnum* EnumType = Cast<UEnum>(Pin->PinType.PinSubCategoryObject.Get());
+				if (EnumType)
+				{
+					// Find the curve mode option
+					FString BestModeValue;
+					for (int32 i = 0; i < EnumType->NumEnums() - 1; ++i) // -1 to skip _MAX
+					{
+						FString EnumName = EnumType->GetNameStringByIndex(i);
+						FString DisplayName = EnumType->GetDisplayNameTextByIndex(i).ToString();
+						
+						UE_LOG(LogTemp, Log, TEXT("SetNiagaraCurve: Scale Mode enum[%d]: '%s' (display: '%s')"), 
+							i, *EnumName, *DisplayName);
+						
+						// Look for curve-related mode
+						// "RGBA Linear Color Curve" contains both "Curve" and "Linear Color"
+						// so we need to prioritize the Curve check
+						bool bIsCurveMode = 
+							DisplayName.Contains(TEXT("Curve")) || 
+							DisplayName.Contains(TEXT("from Curve")) ||
+							EnumName.Contains(TEXT("Curve"));
+						
+						// Skip constant/uniform modes, but NOT if they also contain "Curve"
+						// e.g. "RGBA Linear Color Curve" should be treated as a curve mode
+						bool bIsConstantMode = 
+							!bIsCurveMode && (
+								DisplayName.Contains(TEXT("Linear Color")) ||
+								DisplayName.Equals(TEXT("Uniform"), ESearchCase::IgnoreCase) ||
+								DisplayName.Contains(TEXT("Constant"))
+							);
+						
+						if (bIsCurveMode && !bIsConstantMode)
+						{
+							BestModeValue = EnumName;
+							UE_LOG(LogTemp, Log, TEXT("SetNiagaraCurve: Found curve Scale Mode: '%s'"), *BestModeValue);
+							break;
+						}
+					}
+					
+					// If we found a curve mode and it's different from current value, set it
+					if (!BestModeValue.IsEmpty() && Pin->DefaultValue != BestModeValue)
+					{
+						UE_LOG(LogTemp, Log, TEXT("SetNiagaraCurve: Setting Scale Mode from '%s' to '%s'"), 
+							*Pin->DefaultValue, *BestModeValue);
+						Pin->Modify();
+						Pin->DefaultValue = BestModeValue;
+						
+						// Need to recompile after changing the enum
+						Graph->NotifyGraphChanged();
+						ModuleNode->MarkNodeRequiresSynchronization(TEXT("Scale Mode changed via ECA"), true);
+						NiagaraSystem->MarkPackageDirty();
+						NiagaraSystem->RequestCompile(true);
+						NiagaraSystem->WaitForCompilationComplete();
+						UE_LOG(LogTemp, Log, TEXT("SetNiagaraCurve: Recompiled after Scale Mode change"));
+					}
+				}
+			}
+			break;
+		}
+	}
+	
+	// Build the result
+	TSharedPtr<FJsonObject> Result = MakeResult();
+	Result->SetStringField(TEXT("module_name"), ModuleName);
+	Result->SetStringField(TEXT("input_name"), InputName);
+	Result->SetStringField(TEXT("curve_type"), CurveType);
+	Result->SetNumberField(TEXT("key_count"), KeyCount);
+	Result->SetBoolField(TEXT("curve_created"), true);
+	Result->SetBoolField(TEXT("dynamic_input_used"), true);
+	Result->SetStringField(TEXT("dynamic_input_path"), DynamicInputPath);
+	Result->SetBoolField(TEXT("stateless_emitter"), false);
+	
+	return FECACommandResult::Success(Result);
+#else
+	return FECACommandResult::Error(TEXT("SetNiagaraCurve is only available in editor builds"));
+#endif
+}
+
+// NOTE: Dead code that was previously here has been removed to fix compilation errors.
+// The old approach used if(false) blocks and had undeclared variables.
+
+#if 0 // DEAD CODE START - Everything until #endif was unreachable and caused compilation errors
+	if (CurveType.Equals(TEXT("color"), ESearchCase::IgnoreCase) && 
+		    ModuleName.Contains(TEXT("ScaleColor"), ESearchCase::IgnoreCase))
+		{
+			// Find the Scale Mode pin to discover valid enum values
+			for (UEdGraphPin* Pin : ModuleNode->Pins)
+			{
+				if (Pin->Direction == EGPD_Input && Pin->PinName.ToString().Contains(TEXT("Scale Mode")))
+				{
+					// The enum values for ScaleColor are typically:
+					// - NewEnumerator0 or similar = Uniform/Constant mode
+					// - NewEnumerator1 or similar = Vector from Curve mode
+					// We need to find the "curve" mode - typically it's index 1 or has "Curve" in the name
+					
+					// Check if the pin type has enum info
+					if (Pin->PinType.PinSubCategoryObject.IsValid())
+					{
+						UEnum* EnumType = Cast<UEnum>(Pin->PinType.PinSubCategoryObject.Get());
+						if (EnumType)
+						{
+							// Look for an enum value that suggests "curve" or "vector from curve"
+							for (int32 i = 0; i < EnumType->NumEnums() - 1; ++i) // -1 to skip _MAX
+							{
+								FString EnumName = EnumType->GetNameStringByIndex(i);
+								FString DisplayName = EnumType->GetDisplayNameTextByIndex(i).ToString();
+								
+								UE_LOG(LogTemp, Log, TEXT("SetNiagaraCurve: Scale Mode enum[%d]: '%s' (display: '%s')"), 
+									i, *EnumName, *DisplayName);
+								
+								// Look for curve-related mode
+							// Common patterns in Niagara:
+							// - "Color from Curve" or "Vector from Curve" - these sample from a curve
+							// - "Linear Color" or "Uniform" - these use constant values
+							// We want to select the curve mode, not the constant mode
+							bool bIsCurveMode = 
+								DisplayName.Contains(TEXT("Curve")) || 
+								DisplayName.Contains(TEXT("from Curve")) ||
+								DisplayName.Contains(TEXT("Vector")) ||
+								EnumName.Contains(TEXT("Curve")) ||
+								EnumName.Contains(TEXT("Vector"));
+							
+							// Skip constant/uniform modes
+							bool bIsConstantMode = 
+								DisplayName.Contains(TEXT("Linear Color")) ||
+								DisplayName.Contains(TEXT("Uniform")) ||
+								DisplayName.Contains(TEXT("Constant"));
+							
+							if (bIsCurveMode && !bIsConstantMode)
+							{
+								ScaleModeValue = EnumName;
+								bAutoSetScaleMode = true;
+								UE_LOG(LogTemp, Log, TEXT("SetNiagaraCurve: Auto-selecting Scale Mode: '%s' (display: '%s')"), *ScaleModeValue, *DisplayName);
+								break;
+							}
+							}
+							
+							// If we didn't find a curve mode by name, enumerate all options
+							// and pick the one that's NOT the current constant mode
+							if (!bAutoSetScaleMode && EnumType->NumEnums() > 2)
+							{
+								FString CurrentMode = Pin->DefaultValue;
+								UE_LOG(LogTemp, Log, TEXT("SetNiagaraCurve: Current Scale Mode is '%s', looking for alternative..."), *CurrentMode);
+								
+								// Try to find any mode that isn't the current one and isn't obviously constant
+								for (int32 j = 0; j < EnumType->NumEnums() - 1; ++j)
+								{
+									FString AltEnumName = EnumType->GetNameStringByIndex(j);
+									FString AltDisplayName = EnumType->GetDisplayNameTextByIndex(j).ToString();
+									
+									// Skip the current mode and obvious constant modes
+									if (AltEnumName != CurrentMode && 
+									    !AltDisplayName.Contains(TEXT("Linear Color")) &&
+									    !AltDisplayName.Equals(TEXT("Uniform"), ESearchCase::IgnoreCase))
+									{
+										ScaleModeValue = AltEnumName;
+										bAutoSetScaleMode = true;
+										UE_LOG(LogTemp, Log, TEXT("SetNiagaraCurve: Fallback - selecting '%s' (display: '%s')"), *ScaleModeValue, *AltDisplayName);
+										break;
+									}
+								}
+								
+								// Last resort: just use index 1 if we still haven't found anything
+								if (!bAutoSetScaleMode)
+								{
+									ScaleModeValue = EnumType->GetNameStringByIndex(1);
+									bAutoSetScaleMode = true;
+									UE_LOG(LogTemp, Log, TEXT("SetNiagaraCurve: Last resort fallback to Scale Mode index 1: '%s'"), *ScaleModeValue);
+								}
+							}
+						}
+					}
+					break;
+				}
+			}
+		}
+	}
+	
+	// NOTE: Scale Mode enum is set at the END of this function, after all synchronization
+	// to prevent it from being reset by MarkNodeRequiresSynchronization calls
+	
+	// Create the parameter handle for the input
+	FNiagaraParameterHandle InputHandle = FNiagaraParameterHandle::CreateModuleParameterHandle(FName(*InputName));
+	FNiagaraParameterHandle AliasedHandle = FNiagaraParameterHandle::CreateAliasedModuleParameterHandle(InputHandle, ModuleNode);
+	
+	// Determine the curve data interface type
+	UClass* CurveClass = nullptr;
+	FNiagaraTypeDefinition InputType;
+	
+	if (CurveType.Equals(TEXT("float"), ESearchCase::IgnoreCase))
+	{
+		CurveClass = UNiagaraDataInterfaceCurve::StaticClass();
+		InputType = FNiagaraTypeDefinition(UNiagaraDataInterfaceCurve::StaticClass());
+	}
+	else if (CurveType.Equals(TEXT("vector"), ESearchCase::IgnoreCase))
+	{
+		CurveClass = UNiagaraDataInterfaceVectorCurve::StaticClass();
+		InputType = FNiagaraTypeDefinition(UNiagaraDataInterfaceVectorCurve::StaticClass());
+	}
+	else if (CurveType.Equals(TEXT("color"), ESearchCase::IgnoreCase))
+	{
+		CurveClass = UNiagaraDataInterfaceColorCurve::StaticClass();
+		InputType = FNiagaraTypeDefinition(UNiagaraDataInterfaceColorCurve::StaticClass());
+	}
+	else
+	{
+		return FECACommandResult::Error(FString::Printf(TEXT("Unknown curve type: %s. Supported: float, vector, color"), *CurveType));
+	}
+	
+	// Get or create the override pin for this input
+	UE_LOG(LogTemp, Log, TEXT("SetNiagaraCurve: Creating override pin for input '%s' (aliased: '%s')"), 
+		*InputName, *AliasedHandle.GetParameterHandleString().ToString());
+	
+	UEdGraphPin& OverridePin = FNiagaraStackGraphUtilities::GetOrCreateStackFunctionInputOverridePin(
+		*ModuleNode, AliasedHandle, InputType, FGuid(), FGuid());
+	
+	UE_LOG(LogTemp, Log, TEXT("SetNiagaraCurve: Override pin created/found: '%s', LinkedTo count: %d"), 
+		*OverridePin.PinName.ToString(), OverridePin.LinkedTo.Num());
+	
+	// Check if there's already a data interface linked to this pin that we can modify
+	UNiagaraDataInterface* CurveDataInterface = nullptr;
+	int32 KeyCount = 0;
+	FString AliasedInputName = AliasedHandle.GetParameterHandleString().ToString();
+	
+	if (OverridePin.LinkedTo.Num() > 0)
+	{
+		// Try to get the existing data interface from the linked input node
+		UEdGraphPin* LinkedPin = OverridePin.LinkedTo[0];
+		if (LinkedPin)
+		{
+			UEdGraphNode* LinkedNode = LinkedPin->GetOwningNode();
+			UE_LOG(LogTemp, Log, TEXT("SetNiagaraCurve: Found existing linked node: %s"), *LinkedNode->GetClass()->GetName());
+			
+			// Try to get data interface from the node using reflection
+			// UNiagaraNodeInput is not publicly accessible, so we use reflection
+			FProperty* DataInterfaceProp = LinkedNode->GetClass()->FindPropertyByName(TEXT("DataInterface"));
+			if (DataInterfaceProp)
+			{
+				FObjectProperty* ObjProp = CastField<FObjectProperty>(DataInterfaceProp);
+				if (ObjProp)
+				{
+					CurveDataInterface = Cast<UNiagaraDataInterface>(ObjProp->GetObjectPropertyValue_InContainer(LinkedNode));
+					if (CurveDataInterface)
+					{
+						UE_LOG(LogTemp, Log, TEXT("SetNiagaraCurve: Found existing DataInterface: %s at %p"), 
+							*CurveDataInterface->GetClass()->GetName(), CurveDataInterface);
+					}
+				}
+			}
+			else
+			{
+				UE_LOG(LogTemp, Log, TEXT("SetNiagaraCurve: No DataInterface property found on node"));
+			}
+		}
+		
+		// If we couldn't get the existing data interface, remove the link and create a new one
+		if (!CurveDataInterface)
+		{
+			UE_LOG(LogTemp, Log, TEXT("SetNiagaraCurve: Removing existing linked nodes to create new data interface"));
+			RemoveOverridePinConnections(OverridePin, Graph);
+		}
+	}
+	
+	// If we don't have a data interface yet, create a new one
+	if (!CurveDataInterface)
+	{
+		UE_LOG(LogTemp, Log, TEXT("SetNiagaraCurve: No existing data interface found, creating new one"));
+		UE_LOG(LogTemp, Log, TEXT("SetNiagaraCurve: Calling SetDataInterfaceValueForFunctionInput with class '%s', aliased input name: '%s'"), 
+			*CurveClass->GetName(), *AliasedInputName);
+		
+		FNiagaraStackGraphUtilities::SetDataInterfaceValueForFunctionInput(OverridePin, CurveClass, AliasedInputName, CurveDataInterface, FGuid());
+		
+		UE_LOG(LogTemp, Log, TEXT("SetNiagaraCurve: After SetDataInterfaceValueForFunctionInput - CurveDataInterface=%p, LinkedTo count: %d"), 
+			CurveDataInterface, OverridePin.LinkedTo.Num());
+		
+		// DIAGNOSTIC: Verify the pointer matches what's on the InputNode
+		if (OverridePin.LinkedTo.Num() > 0)
+		{
+			UNiagaraNodeInput* InputNode = Cast<UNiagaraNodeInput>(OverridePin.LinkedTo[0]->GetOwningNode());
+			if (InputNode)
+			{
+				UNiagaraDataInterface* NodeDI = InputNode->GetDataInterface();
+				UE_LOG(LogTemp, Log, TEXT("SetNiagaraCurve: DIAGNOSTIC - InputNode->GetDataInterface()=%p, CurveDataInterface=%p, Match=%s"),
+					NodeDI, CurveDataInterface, (NodeDI == CurveDataInterface) ? TEXT("YES") : TEXT("NO"));
+			}
+		}
+	}
+	else
+	{
+		UE_LOG(LogTemp, Log, TEXT("SetNiagaraCurve: Using existing data interface, will modify in place"));
+	}
+	
+	if (!CurveDataInterface)
+	{
+		return FECACommandResult::Error(TEXT("Failed to create curve data interface"));
+	}
+	
+	// Verify the data interface is the correct type
+	if (!CurveDataInterface->IsA(CurveClass))
+	{
+		UE_LOG(LogTemp, Warning, TEXT("SetNiagaraCurve: Existing data interface is wrong type. Expected %s, got %s. Recreating..."),
+			*CurveClass->GetName(), *CurveDataInterface->GetClass()->GetName());
+		RemoveOverridePinConnections(OverridePin, Graph);
+		FNiagaraStackGraphUtilities::SetDataInterfaceValueForFunctionInput(OverridePin, CurveClass, AliasedInputName, CurveDataInterface, FGuid());
+	}
+	
+	// CRITICAL: Use a scoped transaction to ensure undo/redo works and changes are properly tracked
+	FScopedTransaction Transaction(NSLOCTEXT("ECABridge", "SetNiagaraCurve", "Set Niagara Curve"));
+	
+	// CRITICAL: Call Modify() BEFORE making changes to enable undo/redo and proper serialization
+	CurveDataInterface->Modify();
+	
+	// Now populate the curve with keys using proper property protocol
+	if (CurveType.Equals(TEXT("float"), ESearchCase::IgnoreCase))
+	{
+		UNiagaraDataInterfaceCurve* FloatCurve = Cast<UNiagaraDataInterfaceCurve>(CurveDataInterface);
+		if (FloatCurve)
+		{
+			// CRITICAL: Clear CurveAsset first!
+			// If CurveAsset is set, UpdateLUT() will call SyncCurvesToAsset() which overwrites our curve data!
+			FloatCurve->CurveAsset = nullptr;
+			
+			// Get the actual FProperty pointer
+			FProperty* CurveProp = FloatCurve->GetClass()->FindPropertyByName(TEXT("Curve"));
+			
+			// Call PreEditChange BEFORE modifying
+			if (CurveProp) FloatCurve->PreEditChange(CurveProp);
+			
+			// Clear existing keys
+			FloatCurve->Curve.Reset();
+			
+			for (const TSharedPtr<FJsonValue>& KeyValue : *KeysArray)
+			{
+				const TSharedPtr<FJsonObject>* KeyObj;
+				if (KeyValue->TryGetObject(KeyObj))
+				{
+					float Time = static_cast<float>((*KeyObj)->GetNumberField(TEXT("time")));
+					float Value = static_cast<float>((*KeyObj)->GetNumberField(TEXT("value")));
+					FloatCurve->Curve.AddKey(Time, Value);
+					KeyCount++;
+				}
+			}
+			
+			// Call PostEditChangeProperty with the actual property
+			if (CurveProp)
+			{
+				FPropertyChangedEvent CurveEvent(CurveProp);
+				FloatCurve->PostEditChangeProperty(CurveEvent);
+			}
+			
+			// Update time ranges and LUT after property change notification
+			FloatCurve->UpdateTimeRanges();
+			FloatCurve->UpdateLUT();
+			
+			// Broadcast the OnChanged delegate to notify the input node
+			FloatCurve->OnChanged().Broadcast();
+		}
+	}
+	else if (CurveType.Equals(TEXT("vector"), ESearchCase::IgnoreCase))
+	{
+		UNiagaraDataInterfaceVectorCurve* VectorCurve = Cast<UNiagaraDataInterfaceVectorCurve>(CurveDataInterface);
+		if (VectorCurve)
+		{
+			// CRITICAL: Clear CurveAsset first!
+			// If CurveAsset is set, UpdateLUT() will call SyncCurvesToAsset() which overwrites our curve data!
+			VectorCurve->CurveAsset = nullptr;
+			
+			// Get the actual FProperty pointers
+			FProperty* XCurveProp = VectorCurve->GetClass()->FindPropertyByName(TEXT("XCurve"));
+			FProperty* YCurveProp = VectorCurve->GetClass()->FindPropertyByName(TEXT("YCurve"));
+			FProperty* ZCurveProp = VectorCurve->GetClass()->FindPropertyByName(TEXT("ZCurve"));
+			
+			// Call PreEditChange BEFORE modifying
+			if (XCurveProp) VectorCurve->PreEditChange(XCurveProp);
+			if (YCurveProp) VectorCurve->PreEditChange(YCurveProp);
+			if (ZCurveProp) VectorCurve->PreEditChange(ZCurveProp);
+			
+			// Clear existing keys
+			VectorCurve->XCurve.Reset();
+			VectorCurve->YCurve.Reset();
+			VectorCurve->ZCurve.Reset();
+			
+			for (const TSharedPtr<FJsonValue>& KeyValue : *KeysArray)
+			{
+				const TSharedPtr<FJsonObject>* KeyObj;
+				if (KeyValue->TryGetObject(KeyObj))
+				{
+					float Time = static_cast<float>((*KeyObj)->GetNumberField(TEXT("time")));
+					float X = static_cast<float>((*KeyObj)->GetNumberField(TEXT("x")));
+					float Y = static_cast<float>((*KeyObj)->GetNumberField(TEXT("y")));
+					float Z = static_cast<float>((*KeyObj)->GetNumberField(TEXT("z")));
+					VectorCurve->XCurve.AddKey(Time, X);
+					VectorCurve->YCurve.AddKey(Time, Y);
+					VectorCurve->ZCurve.AddKey(Time, Z);
+					KeyCount++;
+				}
+			}
+			
+			// Call PostEditChangeProperty with the actual properties
+			if (XCurveProp)
+			{
+				FPropertyChangedEvent XEvent(XCurveProp);
+				VectorCurve->PostEditChangeProperty(XEvent);
+			}
+			if (YCurveProp)
+			{
+				FPropertyChangedEvent YEvent(YCurveProp);
+				VectorCurve->PostEditChangeProperty(YEvent);
+			}
+			if (ZCurveProp)
+			{
+				FPropertyChangedEvent ZEvent(ZCurveProp);
+				VectorCurve->PostEditChangeProperty(ZEvent);
+			}
+			
+			// Update time ranges and LUT after property change notification
+			VectorCurve->UpdateTimeRanges();
+			VectorCurve->UpdateLUT();
+			
+			// Broadcast the OnChanged delegate to notify the input node
+			VectorCurve->OnChanged().Broadcast();
+		}
+	}
+	else if (CurveType.Equals(TEXT("color"), ESearchCase::IgnoreCase))
+	{
+		UNiagaraDataInterfaceColorCurve* ColorCurve = Cast<UNiagaraDataInterfaceColorCurve>(CurveDataInterface);
+		if (ColorCurve)
+		{
+			UE_LOG(LogTemp, Log, TEXT("SetNiagaraCurve: Populating ColorCurve with keys using proper property protocol..."));
+			
+			// CRITICAL: Clear CurveAsset first!
+			// If CurveAsset is set, UpdateLUT() will call SyncCurvesToAsset() which overwrites our curve data!
+			ColorCurve->CurveAsset = nullptr;
+			
+			// Get the actual FProperty pointers for each curve - this is essential for proper change notification
+			FProperty* RedCurveProp = ColorCurve->GetClass()->FindPropertyByName(TEXT("RedCurve"));
+			FProperty* GreenCurveProp = ColorCurve->GetClass()->FindPropertyByName(TEXT("GreenCurve"));
+			FProperty* BlueCurveProp = ColorCurve->GetClass()->FindPropertyByName(TEXT("BlueCurve"));
+			FProperty* AlphaCurveProp = ColorCurve->GetClass()->FindPropertyByName(TEXT("AlphaCurve"));
+			
+			UE_LOG(LogTemp, Log, TEXT("  Found properties - Red:%p Green:%p Blue:%p Alpha:%p"),
+				RedCurveProp, GreenCurveProp, BlueCurveProp, AlphaCurveProp);
+			
+			// CRITICAL: Call PreEditChange BEFORE modifying each property
+			// This is how Unreal's property system tracks changes for undo/redo and notifications
+			if (RedCurveProp) ColorCurve->PreEditChange(RedCurveProp);
+			if (GreenCurveProp) ColorCurve->PreEditChange(GreenCurveProp);
+			if (BlueCurveProp) ColorCurve->PreEditChange(BlueCurveProp);
+			if (AlphaCurveProp) ColorCurve->PreEditChange(AlphaCurveProp);
+			
+			// Clear existing keys
+			ColorCurve->RedCurve.Reset();
+			ColorCurve->GreenCurve.Reset();
+			ColorCurve->BlueCurve.Reset();
+			ColorCurve->AlphaCurve.Reset();
+			
+			// Add new keys
+			for (const TSharedPtr<FJsonValue>& KeyValue : *KeysArray)
+			{
+				const TSharedPtr<FJsonObject>* KeyObj;
+				if (KeyValue->TryGetObject(KeyObj))
+				{
+					float Time = static_cast<float>((*KeyObj)->GetNumberField(TEXT("time")));
+					float R = static_cast<float>((*KeyObj)->GetNumberField(TEXT("r")));
+					float G = static_cast<float>((*KeyObj)->GetNumberField(TEXT("g")));
+					float B = static_cast<float>((*KeyObj)->GetNumberField(TEXT("b")));
+					float A = (*KeyObj)->HasField(TEXT("a")) ? static_cast<float>((*KeyObj)->GetNumberField(TEXT("a"))) : 1.0f;
+					
+					// Add keys and verify they were added
+					FKeyHandle RHandle = ColorCurve->RedCurve.AddKey(Time, R);
+					FKeyHandle GHandle = ColorCurve->GreenCurve.AddKey(Time, G);
+					FKeyHandle BHandle = ColorCurve->BlueCurve.AddKey(Time, B);
+					FKeyHandle AHandle = ColorCurve->AlphaCurve.AddKey(Time, A);
+					
+					KeyCount++;
+					UE_LOG(LogTemp, Log, TEXT("  Key %d: Time=%.2f R=%.2f G=%.2f B=%.2f A=%.2f"), KeyCount, Time, R, G, B, A);
+					UE_LOG(LogTemp, Log, TEXT("    Handles valid: R=%d G=%d B=%d A=%d"), 
+						ColorCurve->RedCurve.IsKeyHandleValid(RHandle),
+						ColorCurve->GreenCurve.IsKeyHandleValid(GHandle),
+						ColorCurve->BlueCurve.IsKeyHandleValid(BHandle),
+						ColorCurve->AlphaCurve.IsKeyHandleValid(AHandle));
+					UE_LOG(LogTemp, Log, TEXT("    Current key counts: R=%d G=%d B=%d A=%d"),
+						ColorCurve->RedCurve.GetNumKeys(),
+						ColorCurve->GreenCurve.GetNumKeys(),
+						ColorCurve->BlueCurve.GetNumKeys(),
+						ColorCurve->AlphaCurve.GetNumKeys());
+				}
+			}
+			
+			// CRITICAL: Call PostEditChangeProperty AFTER modifying each property with the ACTUAL property
+			// This triggers all the proper change notifications that the editor uses
+			if (RedCurveProp)
+			{
+				FPropertyChangedEvent RedEvent(RedCurveProp);
+				ColorCurve->PostEditChangeProperty(RedEvent);
+			}
+			if (GreenCurveProp)
+			{
+				FPropertyChangedEvent GreenEvent(GreenCurveProp);
+				ColorCurve->PostEditChangeProperty(GreenEvent);
+			}
+			if (BlueCurveProp)
+			{
+				FPropertyChangedEvent BlueEvent(BlueCurveProp);
+				ColorCurve->PostEditChangeProperty(BlueEvent);
+			}
+			if (AlphaCurveProp)
+			{
+				FPropertyChangedEvent AlphaEvent(AlphaCurveProp);
+				ColorCurve->PostEditChangeProperty(AlphaEvent);
+			}
+			
+			UE_LOG(LogTemp, Log, TEXT("SetNiagaraCurve: BEFORE UpdateLUT - Key counts: R=%d G=%d B=%d A=%d"),
+				ColorCurve->RedCurve.GetNumKeys(),
+				ColorCurve->GreenCurve.GetNumKeys(),
+				ColorCurve->BlueCurve.GetNumKeys(),
+				ColorCurve->AlphaCurve.GetNumKeys());
+			
+			// Update time ranges and LUT after all property changes are notified
+			ColorCurve->UpdateTimeRanges();
+			ColorCurve->UpdateLUT();
+			
+			UE_LOG(LogTemp, Log, TEXT("SetNiagaraCurve: AFTER UpdateLUT - Key counts: R=%d G=%d B=%d A=%d"),
+				ColorCurve->RedCurve.GetNumKeys(),
+				ColorCurve->GreenCurve.GetNumKeys(),
+				ColorCurve->BlueCurve.GetNumKeys(),
+				ColorCurve->AlphaCurve.GetNumKeys());
+			
+			// Check if SyncCurvesToAsset inside UpdateLUT wiped our data
+			if (ColorCurve->RedCurve.GetNumKeys() == 0)
+			{
+				UE_LOG(LogTemp, Error, TEXT("SetNiagaraCurve: CURVES WERE WIPED! CurveAsset=%p"), ColorCurve->CurveAsset.Get());
+			}
+			
+			// Debug: Verify the curves are populated
+			UE_LOG(LogTemp, Log, TEXT("SetNiagaraCurve: ColorCurve populated with %d keys"), KeyCount);
+			
+			// EXTRA DEBUG: Print the actual curve values at a few sample points
+			UE_LOG(LogTemp, Log, TEXT("SetNiagaraCurve: Sampling curve at key times:"));
+			for (int32 i = 0; i < ColorCurve->RedCurve.GetNumKeys(); i++)
+			{
+				float KeyTime = ColorCurve->RedCurve.Keys[i].Time;
+				float R = ColorCurve->RedCurve.Eval(KeyTime);
+				float G = ColorCurve->GreenCurve.Eval(KeyTime);
+				float B = ColorCurve->BlueCurve.Eval(KeyTime);
+				float A = ColorCurve->AlphaCurve.Eval(KeyTime);
+				UE_LOG(LogTemp, Log, TEXT("  Sample at t=%.2f: R=%.2f G=%.2f B=%.2f A=%.2f"), KeyTime, R, G, B, A);
+			}
+			UE_LOG(LogTemp, Log, TEXT("  RedCurve has %d keys, first key time=%.2f value=%.2f"), 
+				ColorCurve->RedCurve.GetNumKeys(),
+				ColorCurve->RedCurve.GetNumKeys() > 0 ? ColorCurve->RedCurve.GetFirstKey().Time : 0.0f,
+				ColorCurve->RedCurve.GetNumKeys() > 0 ? ColorCurve->RedCurve.GetFirstKey().Value : 0.0f);
+			UE_LOG(LogTemp, Log, TEXT("  GreenCurve has %d keys"), ColorCurve->GreenCurve.GetNumKeys());
+			UE_LOG(LogTemp, Log, TEXT("  BlueCurve has %d keys"), ColorCurve->BlueCurve.GetNumKeys());
+			UE_LOG(LogTemp, Log, TEXT("  AlphaCurve has %d keys"), ColorCurve->AlphaCurve.GetNumKeys());
+			UE_LOG(LogTemp, Log, TEXT("  DataInterface address: %p, Outer: %s"), 
+				CurveDataInterface, 
+				CurveDataInterface->GetOuter() ? *CurveDataInterface->GetOuter()->GetName() : TEXT("null"));
+			
+			// Broadcast the OnChanged delegate to notify listeners
+			ColorCurve->OnChanged().Broadcast();
+			
+			// DIAGNOSTIC: Re-verify after all operations that the InputNode still has the same DI with keys
+			if (OverridePin.LinkedTo.Num() > 0)
+			{
+				UNiagaraNodeInput* InputNode = Cast<UNiagaraNodeInput>(OverridePin.LinkedTo[0]->GetOwningNode());
+				if (InputNode)
+				{
+					UNiagaraDataInterfaceColorCurve* NodeColorCurve = Cast<UNiagaraDataInterfaceColorCurve>(InputNode->GetDataInterface());
+					if (NodeColorCurve)
+					{
+						UE_LOG(LogTemp, Log, TEXT("SetNiagaraCurve: FINAL DIAGNOSTIC - InputNode's ColorCurve=%p, RedCurve keys=%d"),
+							NodeColorCurve, NodeColorCurve->RedCurve.GetNumKeys());
+						if (NodeColorCurve != ColorCurve)
+						{
+							UE_LOG(LogTemp, Error, TEXT("SetNiagaraCurve: MISMATCH! InputNode has DIFFERENT ColorCurve than what we modified!"));
+						}
+					}
+				}
+			}
+		}
+		else
+		{
+			UE_LOG(LogTemp, Warning, TEXT("SetNiagaraCurve: Failed to cast CurveDataInterface to UNiagaraDataInterfaceColorCurve"));
+		}
+	}
+	
+	// CRITICAL: Find the input node that owns the data interface and re-set the data interface
+	// This forces the input node to properly register the modified data interface
+	if (OverridePin.LinkedTo.Num() > 0)
+	{
+		UEdGraphPin* LinkedPin = OverridePin.LinkedTo[0];
+		if (LinkedPin)
+		{
+			UEdGraphNode* InputNodeBase = LinkedPin->GetOwningNode();
+			if (InputNodeBase)
+			{
+				InputNodeBase->Modify();
+				UE_LOG(LogTemp, Log, TEXT("SetNiagaraCurve: Found input node '%s'"), *InputNodeBase->GetName());
+				
+				// Cast to UNiagaraNodeInput to call SetDataInterface
+				UNiagaraNodeInput* NiagaraInputNode = Cast<UNiagaraNodeInput>(InputNodeBase);
+				if (NiagaraInputNode)
+				{
+					// NOTE: Don't re-call SetDataInterface - it's the same pointer and we've already modified it
+					// Just mark the node as needing synchronization
+					UE_LOG(LogTemp, Log, TEXT("SetNiagaraCurve: Marking input node for synchronization (NOT re-setting DataInterface)"));
+					NiagaraInputNode->MarkNodeRequiresSynchronization(TEXT("Curve data modified via ECA"), true);
+					
+					// Verify the curves are still populated after all our operations
+					if (UNiagaraDataInterfaceColorCurve* VerifyColorCurve = Cast<UNiagaraDataInterfaceColorCurve>(NiagaraInputNode->GetDataInterface()))
+					{
+						UE_LOG(LogTemp, Log, TEXT("SetNiagaraCurve: VERIFICATION - InputNode's DI RedCurve keys=%d, pointer=%p"),
+							VerifyColorCurve->RedCurve.GetNumKeys(), VerifyColorCurve);
+					}
+				}
+				else
+				{
+					UE_LOG(LogTemp, Warning, TEXT("SetNiagaraCurve: Input node is not a UNiagaraNodeInput, cannot re-set DataInterface"));
+					// Try marking as Niagara node at least
+					if (UNiagaraNode* NiagaraNode = Cast<UNiagaraNode>(InputNodeBase))
+					{
+						NiagaraNode->MarkNodeRequiresSynchronization(TEXT("Curve data modified via ECA"), true);
+					}
+				}
+			}
+		}
+	}
+	
+	// Mark the data interface as modified (should have been called before changes, but call again to be safe)
+	CurveDataInterface->MarkPackageDirty();
+	
+	// CRITICAL: Mark the owning node as requiring synchronization
+	// This ensures the stack view model sees the change
+	UNiagaraNode* OwnerNode = Cast<UNiagaraNode>(OverridePin.GetOwningNode());
+	if (OwnerNode)
+	{
+		OwnerNode->Modify();
+		OwnerNode->MarkNodeRequiresSynchronization(TEXT("Curve input set via ECA"), true);
+	}
+	
+	// Also mark the module node as requiring synchronization
+	ModuleNode->Modify();
+	ModuleNode->MarkNodeRequiresSynchronization(TEXT("Input curve changed via ECA"), true);
+	
+	// Mark graph as modified
+	Graph->Modify();
+	Graph->NotifyGraphChanged();
+	
+	// Mark system as modified and trigger updates
+	NiagaraSystem->Modify();
+	NiagaraSystem->MarkPackageDirty();
+	
+	// CRITICAL: Invalidate the emitter's scripts to force them to be recompiled with our new data interface
+	// Without this, the compilation may use cached script data that doesn't include our new ColorCurve
+	if (Graph)
+	{
+		UE_LOG(LogTemp, Log, TEXT("SetNiagaraCurve: Invalidating graph to force recompile with new DI"));
+		Graph->NotifyGraphNeedsRecompile();
+		Graph->MarkGraphRequiresSynchronization(TEXT("Color curve added via ECA"));
+		Graph->NotifyGraphChanged();
+	}
+	
+	NiagaraSystem->PostEditChange();
+	
+	// CRITICAL: Wait for any pending compile to finish before requesting a new one
+	// This ensures our data interface changes are picked up in the new compile
+	NiagaraSystem->WaitForCompilationComplete();
+	
+	// Request compile with bForce=true to ensure a fresh compile that picks up our changes
+	NiagaraSystem->RequestCompile(true);
+	
+	// Wait for the new compile to complete so the changes are applied
+	NiagaraSystem->WaitForCompilationComplete();
+	
+	UE_LOG(LogTemp, Log, TEXT("SetNiagaraCurve: Synchronous compile completed"));
+	
+	// POST-COMPILE VERIFICATION: Check what's in the compiled script's cached data interfaces
+	UE_LOG(LogTemp, Log, TEXT("SetNiagaraCurve: POST-COMPILE - Checking %d emitters"), NiagaraSystem->GetNumEmitters());
+	for (int32 CheckEmitterIdx = 0; CheckEmitterIdx < NiagaraSystem->GetNumEmitters(); ++CheckEmitterIdx)
+	{
+		FNiagaraEmitterHandle& CheckEmitterHandle = NiagaraSystem->GetEmitterHandle(CheckEmitterIdx);
+		UE_LOG(LogTemp, Log, TEXT("SetNiagaraCurve: POST-COMPILE - Emitter %d: '%s'"), CheckEmitterIdx, *CheckEmitterHandle.GetName().ToString());
+		
+		if (CheckEmitterHandle.GetInstance().Emitter)
+		{
+			const FVersionedNiagaraEmitterData* CheckEmitterData = CheckEmitterHandle.GetInstance().GetEmitterData();
+			if (CheckEmitterData)
+			{
+				// Check spawn script
+				if (UNiagaraScript* SpawnScript = CheckEmitterData->SpawnScriptProps.Script)
+				{
+					UE_LOG(LogTemp, Log, TEXT("SetNiagaraCurve: POST-COMPILE - SpawnScript has %d cached DIs"), SpawnScript->GetCachedDefaultDataInterfaces().Num());
+					for (const FNiagaraScriptDataInterfaceInfo& DIInfo : SpawnScript->GetCachedDefaultDataInterfaces())
+					{
+						UNiagaraDataInterface* DIPtr = DIInfo.DataInterface.Get();
+						UE_LOG(LogTemp, Log, TEXT("SetNiagaraCurve: POST-COMPILE SpawnScript DI: '%s' Type=%s Ptr=%p"),
+							*DIInfo.Name.ToString(), DIPtr ? *DIPtr->GetClass()->GetName() : TEXT("null"), DIPtr);
+						if (UNiagaraDataInterfaceColorCurve* CachedColorCurve = Cast<UNiagaraDataInterfaceColorCurve>(DIPtr))
+						{
+							UE_LOG(LogTemp, Log, TEXT("SetNiagaraCurve: POST-COMPILE SpawnScript ColorCurve '%s' RedCurve keys=%d"),
+								*DIInfo.Name.ToString(), CachedColorCurve->RedCurve.GetNumKeys());
+						}
+					}
+				}
+				// Check update script
+				if (UNiagaraScript* UpdateScript = CheckEmitterData->UpdateScriptProps.Script)
+				{
+					UE_LOG(LogTemp, Log, TEXT("SetNiagaraCurve: POST-COMPILE - UpdateScript has %d cached DIs"), UpdateScript->GetCachedDefaultDataInterfaces().Num());
+					for (const FNiagaraScriptDataInterfaceInfo& DIInfo : UpdateScript->GetCachedDefaultDataInterfaces())
+					{
+						UNiagaraDataInterface* DIPtr = DIInfo.DataInterface.Get();
+						UE_LOG(LogTemp, Log, TEXT("SetNiagaraCurve: POST-COMPILE UpdateScript DI: '%s' Type=%s Ptr=%p"),
+							*DIInfo.Name.ToString(), DIPtr ? *DIPtr->GetClass()->GetName() : TEXT("null"), DIPtr);
+						if (UNiagaraDataInterfaceColorCurve* CachedColorCurve = Cast<UNiagaraDataInterfaceColorCurve>(DIPtr))
+						{
+							UE_LOG(LogTemp, Log, TEXT("SetNiagaraCurve: POST-COMPILE UpdateScript ColorCurve '%s' RedCurve keys=%d"),
+								*DIInfo.Name.ToString(), CachedColorCurve->RedCurve.GetNumKeys());
+						}
+					}
+				}
+			}
+			else
+			{
+				UE_LOG(LogTemp, Warning, TEXT("SetNiagaraCurve: POST-COMPILE - CheckEmitterData is null"));
+			}
+		}
+		else
+		{
+			UE_LOG(LogTemp, Warning, TEXT("SetNiagaraCurve: POST-COMPILE - Emitter instance is null"));
+		}
+	}
+	
+	// IMPORTANT: Set the Scale Mode enum AFTER all synchronization and compile requests
+	// Setting it earlier was getting reset by the synchronization process
+	bool bScaleModeWasSet = false;
+	if ((bSetScaleMode || bAutoSetScaleMode) && !ScaleModeValue.IsEmpty())
+	{
+		for (UEdGraphPin* Pin : ModuleNode->Pins)
+		{
+			if (Pin->Direction == EGPD_Input && 
+			    (Pin->PinName.ToString().Contains(TEXT("Scale Mode")) || 
+			     Pin->PinName.ToString().Contains(TEXT("Mode"))))
+			{
+				UE_LOG(LogTemp, Log, TEXT("SetNiagaraCurve: FINAL Scale Mode set - pin '%s' to '%s' (was '%s')"), 
+					*Pin->PinName.ToString(), *ScaleModeValue, *Pin->DefaultValue);
+				
+				Pin->Modify();
+				Pin->DefaultValue = ScaleModeValue;
+				bScaleModeWasSet = true;
+				break;
+			}
+		}
+	}
+	
+	// CRITICAL FIX: If we changed the Scale Mode, we need to recompile again!
+	// The previous compile was done before setting Scale Mode, so the shader
+	// was compiled with the old mode (constant) instead of the new mode (curve).
+	if (bScaleModeWasSet)
+	{
+		UE_LOG(LogTemp, Log, TEXT("SetNiagaraCurve: Scale Mode was changed, triggering recompile..."));
+		
+		// Mark graph as changed
+		Graph->NotifyGraphChanged();
+		ModuleNode->MarkNodeRequiresSynchronization(TEXT("Scale Mode changed via ECA"), true);
+		
+		// Mark system dirty and request recompile
+		NiagaraSystem->MarkPackageDirty();
+		NiagaraSystem->RequestCompile(true);
+		NiagaraSystem->WaitForCompilationComplete();
+		
+		UE_LOG(LogTemp, Log, TEXT("SetNiagaraCurve: Recompile after Scale Mode change completed"));
+	}
+	
+	UE_LOG(LogTemp, Log, TEXT("SetNiagaraCurve: All nodes marked modified, graph notified, system marked dirty, compile requested"));
+	
+	TSharedPtr<FJsonObject> Result = MakeResult();
+	Result->SetStringField(TEXT("module_name"), ModuleName);
+	Result->SetStringField(TEXT("input_name"), InputName);
+	Result->SetStringField(TEXT("curve_type"), CurveType);
+	Result->SetNumberField(TEXT("key_count"), KeyCount);
+	Result->SetBoolField(TEXT("curve_created"), CurveDataInterface != nullptr);
+	Result->SetBoolField(TEXT("stateless_emitter"), false);
+#endif // DEAD CODE END
+
+//------------------------------------------------------------------------------
+// SetNiagaraDynamicInput - Set dynamic input on a module
+//------------------------------------------------------------------------------
+
+FECACommandResult FECACommand_SetNiagaraDynamicInput::Execute(const TSharedPtr<FJsonObject>& Params)
+{
+#if WITH_EDITORONLY_DATA
+	FString SystemPath;
+	if (!GetStringParam(Params, TEXT("system_path"), SystemPath))
+	{
+		return FECACommandResult::Error(TEXT("Missing required parameter: system_path"));
+	}
+	
+	FString EmitterName;
+	if (!GetStringParam(Params, TEXT("emitter_name"), EmitterName))
+	{
+		return FECACommandResult::Error(TEXT("Missing required parameter: emitter_name"));
+	}
+	
+	FString ModuleName;
+	if (!GetStringParam(Params, TEXT("module_name"), ModuleName))
+	{
+		return FECACommandResult::Error(TEXT("Missing required parameter: module_name"));
+	}
+	
+	FString InputName;
+	if (!GetStringParam(Params, TEXT("input_name"), InputName))
+	{
+		return FECACommandResult::Error(TEXT("Missing required parameter: input_name"));
+	}
+	
+	FString ScriptUsageStr;
+	if (!GetStringParam(Params, TEXT("script_usage"), ScriptUsageStr))
+	{
+		return FECACommandResult::Error(TEXT("Missing required parameter: script_usage"));
+	}
+	
+	FString DynamicInputType;
+	if (!GetStringParam(Params, TEXT("dynamic_input_type"), DynamicInputType))
+	{
+		return FECACommandResult::Error(TEXT("Missing required parameter: dynamic_input_type"));
+	}
+	
+	// Load the system
+	UNiagaraSystem* NiagaraSystem = LoadObject<UNiagaraSystem>(nullptr, *SystemPath);
+	if (!NiagaraSystem)
+	{
+		return FECACommandResult::Error(FString::Printf(TEXT("Failed to load Niagara system: %s"), *SystemPath));
+	}
+	
+	// Find the emitter
+	FNiagaraEmitterHandle* EmitterHandle = FindEmitterHandleByName(NiagaraSystem, EmitterName);
+	if (!EmitterHandle)
+	{
+		return FECACommandResult::Error(FString::Printf(TEXT("Emitter not found: %s"), *EmitterName));
+	}
+	
+	FVersionedNiagaraEmitterData* EmitterData = EmitterHandle->GetEmitterData();
+	if (!EmitterData)
+	{
+		return FECACommandResult::Error(TEXT("Emitter has no data"));
+	}
+	
+	// Get the script source and graph
+	UNiagaraScriptSource* ScriptSource = Cast<UNiagaraScriptSource>(EmitterData->GraphSource);
+	if (!ScriptSource || !ScriptSource->NodeGraph)
+	{
+		return FECACommandResult::Error(TEXT("Emitter has no script source"));
+	}
+	
+	UNiagaraGraph* Graph = ScriptSource->NodeGraph;
+	
+	// Find the module node
+	TArray<UNiagaraNodeFunctionCall*> FunctionNodes;
+	Graph->GetNodesOfClass<UNiagaraNodeFunctionCall>(FunctionNodes);
+	
+	UNiagaraNodeFunctionCall* ModuleNode = nullptr;
+	ENiagaraScriptUsage TargetUsage = ParseScriptUsage(ScriptUsageStr);
+	
+	for (UNiagaraNodeFunctionCall* FuncNode : FunctionNodes)
+	{
+		if (FuncNode->GetFunctionName().Equals(ModuleName, ESearchCase::IgnoreCase))
+		{
+			ENiagaraScriptUsage NodeUsage = FNiagaraStackGraphUtilities::GetOutputNodeUsage(*FuncNode);
+			if (NodeUsage == TargetUsage)
+			{
+				ModuleNode = FuncNode;
+				break;
+			}
+		}
+	}
+	
+	if (!ModuleNode)
+	{
+		return FECACommandResult::Error(FString::Printf(TEXT("Module not found: %s in %s"), *ModuleName, *ScriptUsageStr));
+	}
+	
+	// Create the parameter handle for the input
+	FNiagaraParameterHandle InputHandle = FNiagaraParameterHandle::CreateModuleParameterHandle(FName(*InputName));
+	FNiagaraParameterHandle AliasedHandle = FNiagaraParameterHandle::CreateAliasedModuleParameterHandle(InputHandle, ModuleNode);
+	
+	TSharedPtr<FJsonObject> Result = MakeResult();
+	Result->SetStringField(TEXT("module_name"), ModuleName);
+	Result->SetStringField(TEXT("input_name"), InputName);
+	Result->SetStringField(TEXT("dynamic_input_type"), DynamicInputType);
+	
+	// Handle different dynamic input types
+	if (DynamicInputType.Equals(TEXT("random_range"), ESearchCase::IgnoreCase) ||
+		DynamicInputType.Equals(TEXT("uniform_random"), ESearchCase::IgnoreCase))
+	{
+		// Load the random range dynamic input script
+		FString RandomScriptPath = TEXT("/Niagara/Modules/DynamicInputs/UniformRangedFloat.UniformRangedFloat");
+		
+		// Check if we have vector values
+		FVector MinVec, MaxVec;
+		double MinFloat = 0.0, MaxFloat = 1.0;
+		bool bIsVector = false;
+		
+		if (GetVectorParam(Params, TEXT("min_value"), MinVec, false) && GetVectorParam(Params, TEXT("max_value"), MaxVec, false))
+		{
+			bIsVector = true;
+			RandomScriptPath = TEXT("/Niagara/Modules/DynamicInputs/UniformRangedVector.UniformRangedVector");
+		}
+		else
+		{
+			GetFloatParam(Params, TEXT("min_value"), MinFloat, false);
+			GetFloatParam(Params, TEXT("max_value"), MaxFloat, false);
+		}
+		
+		// Load the dynamic input script
+		UNiagaraScript* DynamicInputScript = LoadObject<UNiagaraScript>(nullptr, *RandomScriptPath);
+		if (DynamicInputScript)
+		{
+			// Get or create the override pin
+			FNiagaraTypeDefinition InputType = bIsVector ? FNiagaraTypeDefinition::GetVec3Def() : FNiagaraTypeDefinition::GetFloatDef();
+			UEdGraphPin& OverridePin = FNiagaraStackGraphUtilities::GetOrCreateStackFunctionInputOverridePin(
+				*ModuleNode, AliasedHandle, InputType, FGuid(), FGuid());
+			
+			// Set the dynamic input
+			UNiagaraNodeFunctionCall* DynamicInputNode = nullptr;
+			FNiagaraStackGraphUtilities::SetDynamicInputForFunctionInput(OverridePin, DynamicInputScript, DynamicInputNode, FGuid(), TEXT("Random"), FGuid());
+			
+			Result->SetStringField(TEXT("script_path"), RandomScriptPath);
+			Result->SetBoolField(TEXT("dynamic_input_created"), DynamicInputNode != nullptr);
+		}
+		else
+		{
+			Result->SetBoolField(TEXT("dynamic_input_created"), false);
+			Result->SetStringField(TEXT("error"), TEXT("Failed to load random range script"));
+		}
+	}
+	else if (DynamicInputType.Equals(TEXT("parameter_link"), ESearchCase::IgnoreCase))
+	{
+		// Note: GetParametersForContext is not exported from NiagaraEditor module
+		// This feature requires internal Niagara editor APIs that are not available to external plugins
+		return FECACommandResult::Error(TEXT("parameter_link dynamic input type is not supported - GetParametersForContext is not exported from NiagaraEditor"));
+	}
+	else if (DynamicInputType.Equals(TEXT("custom_expression"), ESearchCase::IgnoreCase))
+	{
+		// Note: SetCustomExpressionForFunctionInput is not exported from NiagaraEditor module
+		// This feature requires internal Niagara editor APIs that are not available to external plugins
+		return FECACommandResult::Error(TEXT("custom_expression dynamic input type is not supported - SetCustomExpressionForFunctionInput is not exported from NiagaraEditor"));
+	}
+	else
+	{
+		return FECACommandResult::Error(FString::Printf(TEXT("Unknown dynamic input type: %s. Supported: random_range, uniform_random"), *DynamicInputType));
+	}
+	
+	// Mark as modified
+	Graph->NotifyGraphChanged();
+	NiagaraSystem->MarkPackageDirty();
+	NiagaraSystem->RequestCompile(false);
+	NiagaraSystem->PostEditChange();
+	
+	return FECACommandResult::Success(Result);
+#else
+	return FECACommandResult::Error(TEXT("SetNiagaraDynamicInput is only available in editor builds"));
+#endif
+}
