ECABridge experimental MCP plugin

diff --git a/Engine/Plugins/Experimental/ECABridge/Source/ECABridge/Private/Commands/ECAProjectCommands.cpp b/Engine/Plugins/Experimental/ECABridge/Source/ECABridge/Private/Commands/ECAProjectCommands.cpp
new file mode 100644
--- /dev/null
+++ b/Engine/Plugins/Experimental/ECABridge/Source/ECABridge/Private/Commands/ECAProjectCommands.cpp
@@ -0,0 +1,349 @@
+// Copyright Epic Games, Inc. All Rights Reserved.
+
+#include "Commands/ECAProjectCommands.h"
+#include "Commands/ECACommand.h"
+#include "GameFramework/InputSettings.h"
+#include "GameMapsSettings.h"
+#include "Engine/Engine.h"
+#include "Misc/App.h"
+#include "Misc/EngineVersion.h"
+
+// Register all project commands
+REGISTER_ECA_COMMAND(FECACommand_CreateInputMapping)
+REGISTER_ECA_COMMAND(FECACommand_GetInputMappings)
+REGISTER_ECA_COMMAND(FECACommand_RemoveInputMapping)
+REGISTER_ECA_COMMAND(FECACommand_GetProjectSettings)
+REGISTER_ECA_COMMAND(FECACommand_SetProjectSetting)
+
+//------------------------------------------------------------------------------
+// Helper: Parse key name to FKey
+//------------------------------------------------------------------------------
+
+static FKey ParseKeyName(const FString& KeyName)
+{
+	// Try direct parsing first
+	FKey Key = FKey(*KeyName);
+	if (Key.IsValid())
+	{
+		return Key;
+	}
+	
+	// Try common aliases
+	static TMap<FString, FKey> KeyAliases = {
+		// Mouse
+		{ TEXT("LeftMouse"), EKeys::LeftMouseButton },
+		{ TEXT("RightMouse"), EKeys::RightMouseButton },
+		{ TEXT("MiddleMouse"), EKeys::MiddleMouseButton },
+		{ TEXT("MouseX"), EKeys::MouseX },
+		{ TEXT("MouseY"), EKeys::MouseY },
+		{ TEXT("MouseWheelUp"), EKeys::MouseScrollUp },
+		{ TEXT("MouseWheelDown"), EKeys::MouseScrollDown },
+		
+		// Keyboard shortcuts
+		{ TEXT("Space"), EKeys::SpaceBar },
+		{ TEXT("Enter"), EKeys::Enter },
+		{ TEXT("Escape"), EKeys::Escape },
+		{ TEXT("Tab"), EKeys::Tab },
+		{ TEXT("Backspace"), EKeys::BackSpace },
+		
+		// Gamepad
+		{ TEXT("GamepadA"), EKeys::Gamepad_FaceButton_Bottom },
+		{ TEXT("GamepadB"), EKeys::Gamepad_FaceButton_Right },
+		{ TEXT("GamepadX"), EKeys::Gamepad_FaceButton_Left },
+		{ TEXT("GamepadY"), EKeys::Gamepad_FaceButton_Top },
+		{ TEXT("GamepadLeftStickX"), EKeys::Gamepad_LeftX },
+		{ TEXT("GamepadLeftStickY"), EKeys::Gamepad_LeftY },
+		{ TEXT("GamepadRightStickX"), EKeys::Gamepad_RightX },
+		{ TEXT("GamepadRightStickY"), EKeys::Gamepad_RightY },
+	};
+	
+	const FKey* Found = KeyAliases.Find(KeyName);
+	if (Found)
+	{
+		return *Found;
+	}
+	
+	// Return the parsed key even if not fully valid - UE will handle it
+	return FKey(*KeyName);
+}
+
+//------------------------------------------------------------------------------
+// CreateInputMapping
+//------------------------------------------------------------------------------
+
+FECACommandResult FECACommand_CreateInputMapping::Execute(const TSharedPtr<FJsonObject>& Params)
+{
+	FString ActionName;
+	if (!GetStringParam(Params, TEXT("action_name"), ActionName))
+	{
+		return FECACommandResult::Error(TEXT("Missing required parameter: action_name"));
+	}
+	
+	FString KeyName;
+	if (!GetStringParam(Params, TEXT("key"), KeyName))
+	{
+		return FECACommandResult::Error(TEXT("Missing required parameter: key"));
+	}
+	
+	FString InputType = TEXT("Action");
+	GetStringParam(Params, TEXT("input_type"), InputType, false);
+	
+	double Scale = 1.0;
+	GetFloatParam(Params, TEXT("scale"), Scale, false);
+	
+	bool bShift = false, bCtrl = false, bAlt = false, bCmd = false;
+	GetBoolParam(Params, TEXT("shift"), bShift, false);
+	GetBoolParam(Params, TEXT("ctrl"), bCtrl, false);
+	GetBoolParam(Params, TEXT("alt"), bAlt, false);
+	GetBoolParam(Params, TEXT("cmd"), bCmd, false);
+	
+	UInputSettings* InputSettings = UInputSettings::GetInputSettings();
+	if (!InputSettings)
+	{
+		return FECACommandResult::Error(TEXT("Failed to get Input Settings"));
+	}
+	
+	FKey Key = ParseKeyName(KeyName);
+	
+	if (InputType.Equals(TEXT("Action"), ESearchCase::IgnoreCase))
+	{
+		// Create action mapping
+		FInputActionKeyMapping ActionMapping;
+		ActionMapping.ActionName = FName(*ActionName);
+		ActionMapping.Key = Key;
+		ActionMapping.bShift = bShift;
+		ActionMapping.bCtrl = bCtrl;
+		ActionMapping.bAlt = bAlt;
+		ActionMapping.bCmd = bCmd;
+		
+		InputSettings->AddActionMapping(ActionMapping);
+	}
+	else if (InputType.Equals(TEXT("Axis"), ESearchCase::IgnoreCase))
+	{
+		// Create axis mapping
+		FInputAxisKeyMapping AxisMapping;
+		AxisMapping.AxisName = FName(*ActionName);
+		AxisMapping.Key = Key;
+		AxisMapping.Scale = Scale;
+		
+		InputSettings->AddAxisMapping(AxisMapping);
+	}
+	else
+	{
+		return FECACommandResult::Error(FString::Printf(TEXT("Invalid input_type: %s (must be Action or Axis)"), *InputType));
+	}
+	
+	// Save settings
+	InputSettings->SaveKeyMappings();
+	InputSettings->TryUpdateDefaultConfigFile();
+	
+	TSharedPtr<FJsonObject> Result = MakeResult();
+	Result->SetStringField(TEXT("action_name"), ActionName);
+	Result->SetStringField(TEXT("key"), Key.ToString());
+	Result->SetStringField(TEXT("input_type"), InputType);
+	return FECACommandResult::Success(Result);
+}
+
+//------------------------------------------------------------------------------
+// GetInputMappings
+//------------------------------------------------------------------------------
+
+FECACommandResult FECACommand_GetInputMappings::Execute(const TSharedPtr<FJsonObject>& Params)
+{
+	FString InputType = TEXT("All");
+	GetStringParam(Params, TEXT("input_type"), InputType, false);
+	
+	UInputSettings* InputSettings = UInputSettings::GetInputSettings();
+	if (!InputSettings)
+	{
+		return FECACommandResult::Error(TEXT("Failed to get Input Settings"));
+	}
+	
+	TSharedPtr<FJsonObject> Result = MakeResult();
+	
+	// Action mappings
+	if (InputType.Equals(TEXT("All"), ESearchCase::IgnoreCase) || InputType.Equals(TEXT("Action"), ESearchCase::IgnoreCase))
+	{
+		TArray<TSharedPtr<FJsonValue>> ActionMappingsArray;
+		
+		const TArray<FInputActionKeyMapping>& ActionMappings = InputSettings->GetActionMappings();
+		for (const FInputActionKeyMapping& Mapping : ActionMappings)
+		{
+			TSharedPtr<FJsonObject> MappingJson = MakeShared<FJsonObject>();
+			MappingJson->SetStringField(TEXT("action_name"), Mapping.ActionName.ToString());
+			MappingJson->SetStringField(TEXT("key"), Mapping.Key.ToString());
+			MappingJson->SetBoolField(TEXT("shift"), Mapping.bShift);
+			MappingJson->SetBoolField(TEXT("ctrl"), Mapping.bCtrl);
+			MappingJson->SetBoolField(TEXT("alt"), Mapping.bAlt);
+			MappingJson->SetBoolField(TEXT("cmd"), Mapping.bCmd);
+			ActionMappingsArray.Add(MakeShared<FJsonValueObject>(MappingJson));
+		}
+		
+		Result->SetArrayField(TEXT("action_mappings"), ActionMappingsArray);
+	}
+	
+	// Axis mappings
+	if (InputType.Equals(TEXT("All"), ESearchCase::IgnoreCase) || InputType.Equals(TEXT("Axis"), ESearchCase::IgnoreCase))
+	{
+		TArray<TSharedPtr<FJsonValue>> AxisMappingsArray;
+		
+		const TArray<FInputAxisKeyMapping>& AxisMappings = InputSettings->GetAxisMappings();
+		for (const FInputAxisKeyMapping& Mapping : AxisMappings)
+		{
+			TSharedPtr<FJsonObject> MappingJson = MakeShared<FJsonObject>();
+			MappingJson->SetStringField(TEXT("axis_name"), Mapping.AxisName.ToString());
+			MappingJson->SetStringField(TEXT("key"), Mapping.Key.ToString());
+			MappingJson->SetNumberField(TEXT("scale"), Mapping.Scale);
+			AxisMappingsArray.Add(MakeShared<FJsonValueObject>(MappingJson));
+		}
+		
+		Result->SetArrayField(TEXT("axis_mappings"), AxisMappingsArray);
+	}
+	
+	return FECACommandResult::Success(Result);
+}
+
+//------------------------------------------------------------------------------
+// RemoveInputMapping
+//------------------------------------------------------------------------------
+
+FECACommandResult FECACommand_RemoveInputMapping::Execute(const TSharedPtr<FJsonObject>& Params)
+{
+	FString ActionName;
+	if (!GetStringParam(Params, TEXT("action_name"), ActionName))
+	{
+		return FECACommandResult::Error(TEXT("Missing required parameter: action_name"));
+	}
+	
+	FString InputType = TEXT("Action");
+	GetStringParam(Params, TEXT("input_type"), InputType, false);
+	
+	UInputSettings* InputSettings = UInputSettings::GetInputSettings();
+	if (!InputSettings)
+	{
+		return FECACommandResult::Error(TEXT("Failed to get Input Settings"));
+	}
+	
+	int32 RemovedCount = 0;
+	
+	if (InputType.Equals(TEXT("Action"), ESearchCase::IgnoreCase))
+	{
+		// Get current mappings
+		TArray<FInputActionKeyMapping> Mappings;
+		InputSettings->GetActionMappingByName(FName(*ActionName), Mappings);
+		
+		for (const FInputActionKeyMapping& Mapping : Mappings)
+		{
+			InputSettings->RemoveActionMapping(Mapping);
+			RemovedCount++;
+		}
+	}
+	else if (InputType.Equals(TEXT("Axis"), ESearchCase::IgnoreCase))
+	{
+		TArray<FInputAxisKeyMapping> Mappings;
+		InputSettings->GetAxisMappingByName(FName(*ActionName), Mappings);
+		
+		for (const FInputAxisKeyMapping& Mapping : Mappings)
+		{
+			InputSettings->RemoveAxisMapping(Mapping);
+			RemovedCount++;
+		}
+	}
+	
+	if (RemovedCount > 0)
+	{
+		InputSettings->SaveKeyMappings();
+		InputSettings->TryUpdateDefaultConfigFile();
+	}
+	
+	TSharedPtr<FJsonObject> Result = MakeResult();
+	Result->SetStringField(TEXT("action_name"), ActionName);
+	Result->SetNumberField(TEXT("removed_count"), RemovedCount);
+	return FECACommandResult::Success(Result);
+}
+
+//------------------------------------------------------------------------------
+// GetProjectSettings
+//------------------------------------------------------------------------------
+
+FECACommandResult FECACommand_GetProjectSettings::Execute(const TSharedPtr<FJsonObject>& Params)
+{
+	TSharedPtr<FJsonObject> Result = MakeResult();
+	
+	// Game Maps settings
+	UGameMapsSettings* GameMapsSettings = GetMutableDefault<UGameMapsSettings>();
+	if (GameMapsSettings)
+	{
+		Result->SetStringField(TEXT("game_default_map"), GameMapsSettings->GetGameDefaultMap());
+		Result->SetStringField(TEXT("editor_startup_map"), GameMapsSettings->EditorStartupMap.ToString());
+		Result->SetStringField(TEXT("game_instance_class"), GameMapsSettings->GameInstanceClass.ToString());
+		Result->SetStringField(TEXT("global_default_game_mode"), GameMapsSettings->GetGlobalDefaultGameMode());
+	}
+	
+	// General project info
+	Result->SetStringField(TEXT("project_name"), FApp::GetProjectName());
+	Result->SetStringField(TEXT("engine_version"), FEngineVersion::Current().ToString());
+	
+	return FECACommandResult::Success(Result);
+}
+
+//------------------------------------------------------------------------------
+// SetProjectSetting
+//------------------------------------------------------------------------------
+
+FECACommandResult FECACommand_SetProjectSetting::Execute(const TSharedPtr<FJsonObject>& Params)
+{
+	FString SettingName;
+	if (!GetStringParam(Params, TEXT("setting_name"), SettingName))
+	{
+		return FECACommandResult::Error(TEXT("Missing required parameter: setting_name"));
+	}
+	
+	FString SettingValue;
+	if (!GetStringParam(Params, TEXT("setting_value"), SettingValue))
+	{
+		return FECACommandResult::Error(TEXT("Missing required parameter: setting_value"));
+	}
+	
+	UGameMapsSettings* GameMapsSettings = GetMutableDefault<UGameMapsSettings>();
+	if (!GameMapsSettings)
+	{
+		return FECACommandResult::Error(TEXT("Failed to get Game Maps Settings"));
+	}
+	
+	bool bSettingFound = false;
+	
+	if (SettingName.Equals(TEXT("GameDefaultMap"), ESearchCase::IgnoreCase))
+	{
+		GameMapsSettings->SetGameDefaultMap(SettingValue);
+		bSettingFound = true;
+	}
+	else if (SettingName.Equals(TEXT("EditorStartupMap"), ESearchCase::IgnoreCase))
+	{
+		GameMapsSettings->EditorStartupMap = FSoftObjectPath(SettingValue);
+		bSettingFound = true;
+	}
+	else if (SettingName.Equals(TEXT("GlobalDefaultGameMode"), ESearchCase::IgnoreCase))
+	{
+		GameMapsSettings->SetGlobalDefaultGameMode(SettingValue);
+		bSettingFound = true;
+	}
+	else if (SettingName.Equals(TEXT("GameInstanceClass"), ESearchCase::IgnoreCase))
+	{
+		GameMapsSettings->GameInstanceClass = FSoftClassPath(SettingValue);
+		bSettingFound = true;
+	}
+	
+	if (!bSettingFound)
+	{
+		return FECACommandResult::Error(FString::Printf(TEXT("Unknown setting: %s"), *SettingName));
+	}
+	
+	// Save config
+	GameMapsSettings->TryUpdateDefaultConfigFile();
+	
+	TSharedPtr<FJsonObject> Result = MakeResult();
+	Result->SetStringField(TEXT("setting_name"), SettingName);
+	Result->SetStringField(TEXT("setting_value"), SettingValue);
+	return FECACommandResult::Success(Result);
+}
diff --git a/Engine/Plugins/Experimental/ECABridge/Source/ECABridge/Private/Commands/ECAUMGCommands.cpp b/Engine/Plugins/Experimental/ECABridge/Source/ECABridge/Private/Commands/ECAUMGCommands.cpp
new file mode 100644
--- /dev/null
+++ b/Engine/Plugins/Experimental/ECABridge/Source/ECABridge/Private/Commands/ECAUMGCommands.cpp
@@ -0,0 +1,621 @@
+// Copyright Epic Games, Inc. All Rights Reserved.
+
+#include "Commands/ECAUMGCommands.h"
+#include "Commands/ECACommand.h"
+#include "Editor.h"
+#include "Engine/Blueprint.h"
+#include "WidgetBlueprint.h"
+#include "Blueprint/WidgetTree.h"
+#include "Blueprint/UserWidget.h"
+#include "Components/TextBlock.h"
+#include "Components/Button.h"
+#include "Components/Image.h"
+#include "Components/CanvasPanel.h"
+#include "Components/CanvasPanelSlot.h"
+#include "Kismet2/KismetEditorUtilities.h"
+#include "Kismet2/BlueprintEditorUtils.h"
+#include "AssetRegistry/AssetRegistryModule.h"
+#include "Engine/Texture2D.h"
+
+// Register all UMG commands
+REGISTER_ECA_COMMAND(FECACommand_CreateUMGWidgetBlueprint)
+REGISTER_ECA_COMMAND(FECACommand_AddTextBlockToWidget)
+REGISTER_ECA_COMMAND(FECACommand_AddButtonToWidget)
+REGISTER_ECA_COMMAND(FECACommand_AddWidgetToViewport)
+REGISTER_ECA_COMMAND(FECACommand_BindWidgetEvent)
+REGISTER_ECA_COMMAND(FECACommand_SetTextBlockBinding)
+REGISTER_ECA_COMMAND(FECACommand_AddImageToWidget)
+REGISTER_ECA_COMMAND(FECACommand_GetWidgetInfo)
+
+//------------------------------------------------------------------------------
+// Helper Functions
+//------------------------------------------------------------------------------
+
+static UWidgetBlueprint* LoadUMGWidgetBlueprintByPath(const FString& WidgetPath)
+{
+	return LoadObject<UWidgetBlueprint>(nullptr, *WidgetPath);
+}
+
+static FVector2D GetPositionFromParams(const TSharedPtr<FJsonObject>& Params)
+{
+	FVector2D Position(0, 0);
+	const TSharedPtr<FJsonObject>* PosObj;
+	if (Params->TryGetObjectField(TEXT("position"), PosObj))
+	{
+		Position.X = (*PosObj)->GetNumberField(TEXT("x"));
+		Position.Y = (*PosObj)->GetNumberField(TEXT("y"));
+	}
+	return Position;
+}
+
+static FVector2D GetSizeFromParams(const TSharedPtr<FJsonObject>& Params, FVector2D DefaultSize = FVector2D(200, 50))
+{
+	FVector2D Size = DefaultSize;
+	const TSharedPtr<FJsonObject>* SizeObj;
+	if (Params->TryGetObjectField(TEXT("size"), SizeObj))
+	{
+		Size.X = (*SizeObj)->GetNumberField(TEXT("width"));
+		Size.Y = (*SizeObj)->GetNumberField(TEXT("height"));
+	}
+	return Size;
+}
+
+static FLinearColor GetColorFromParams(const TSharedPtr<FJsonObject>& Params, const FString& ParamName, FLinearColor DefaultColor = FLinearColor::White)
+{
+	FLinearColor Color = DefaultColor;
+	const TSharedPtr<FJsonObject>* ColorObj;
+	if (Params->TryGetObjectField(*ParamName, ColorObj))
+	{
+		Color.R = (*ColorObj)->GetNumberField(TEXT("r"));
+		Color.G = (*ColorObj)->GetNumberField(TEXT("g"));
+		Color.B = (*ColorObj)->GetNumberField(TEXT("b"));
+		double Alpha = 1.0;
+		(*ColorObj)->TryGetNumberField(TEXT("a"), Alpha);
+		Color.A = Alpha;
+	}
+	return Color;
+}
+
+//------------------------------------------------------------------------------
+// CreateUMGWidgetBlueprint
+//------------------------------------------------------------------------------
+
+FECACommandResult FECACommand_CreateUMGWidgetBlueprint::Execute(const TSharedPtr<FJsonObject>& Params)
+{
+	FString WidgetName;
+	if (!GetStringParam(Params, TEXT("widget_name"), WidgetName))
+	{
+		return FECACommandResult::Error(TEXT("Missing required parameter: widget_name"));
+	}
+	
+	FString ParentClassStr = TEXT("UserWidget");
+	GetStringParam(Params, TEXT("parent_class"), ParentClassStr, false);
+	
+	FString Path = TEXT("/Game/UI/");
+	GetStringParam(Params, TEXT("path"), Path, false);
+	
+	if (!Path.EndsWith(TEXT("/")))
+	{
+		Path += TEXT("/");
+	}
+	
+	// Get parent class
+	UClass* ParentClass = UUserWidget::StaticClass();
+	if (!ParentClassStr.Equals(TEXT("UserWidget"), ESearchCase::IgnoreCase))
+	{
+		UClass* FoundClass = FindObject<UClass>(nullptr, *FString::Printf(TEXT("/Script/UMG.%s"), *ParentClassStr));
+		if (FoundClass && FoundClass->IsChildOf(UUserWidget::StaticClass()))
+		{
+			ParentClass = FoundClass;
+		}
+	}
+	
+	// Create package path
+	FString PackagePath = Path + WidgetName;
+	
+	// Check if asset already exists at this path
+	if (UObject* ExistingAsset = StaticLoadObject(UObject::StaticClass(), nullptr, *PackagePath))
+	{
+		// Asset already exists - return info about it instead of crashing
+		if (UWidgetBlueprint* ExistingWidget = Cast<UWidgetBlueprint>(ExistingAsset))
+		{
+			TSharedPtr<FJsonObject> Result = MakeResult();
+			Result->SetStringField(TEXT("widget_path"), PackagePath);
+			Result->SetStringField(TEXT("widget_name"), WidgetName);
+			Result->SetStringField(TEXT("parent_class"), ExistingWidget->ParentClass ? ExistingWidget->ParentClass->GetName() : TEXT("Unknown"));
+			Result->SetBoolField(TEXT("already_exists"), true);
+			Result->SetStringField(TEXT("message"), TEXT("Widget Blueprint already exists at this path"));
+			return FECACommandResult::Success(Result);
+		}
+		else
+		{
+			return FECACommandResult::Error(FString::Printf(TEXT("Asset already exists at path '%s' but is not a Widget Blueprint (it's a %s)"), *PackagePath, *ExistingAsset->GetClass()->GetName()));
+		}
+	}
+	
+	// Create package
+	UPackage* Package = CreatePackage(*PackagePath);
+	if (!Package)
+	{
+		return FECACommandResult::Error(TEXT("Failed to create package"));
+	}
+	
+	// Create Widget Blueprint
+	UWidgetBlueprint* WidgetBlueprint = CastChecked<UWidgetBlueprint>(
+		FKismetEditorUtilities::CreateBlueprint(
+			ParentClass,
+			Package,
+			*WidgetName,
+			BPTYPE_Normal,
+			UWidgetBlueprint::StaticClass(),
+			UBlueprintGeneratedClass::StaticClass()
+		)
+	);
+	
+	if (!WidgetBlueprint)
+	{
+		return FECACommandResult::Error(TEXT("Failed to create Widget Blueprint"));
+	}
+	
+	// Add a default Canvas Panel as root
+	UWidgetTree* WidgetTree = WidgetBlueprint->WidgetTree;
+	if (WidgetTree)
+	{
+		UCanvasPanel* RootCanvas = WidgetTree->ConstructWidget<UCanvasPanel>(UCanvasPanel::StaticClass(), TEXT("RootCanvas"));
+		WidgetTree->RootWidget = RootCanvas;
+	}
+	
+	// Notify asset registry
+	FAssetRegistryModule::AssetCreated(WidgetBlueprint);
+	Package->MarkPackageDirty();
+	
+	// Compile
+	FKismetEditorUtilities::CompileBlueprint(WidgetBlueprint);
+	
+	TSharedPtr<FJsonObject> Result = MakeResult();
+	Result->SetStringField(TEXT("widget_path"), PackagePath);
+	Result->SetStringField(TEXT("widget_name"), WidgetName);
+	Result->SetStringField(TEXT("parent_class"), ParentClass->GetName());
+	return FECACommandResult::Success(Result);
+}
+
+//------------------------------------------------------------------------------
+// AddTextBlockToWidget
+//------------------------------------------------------------------------------
+
+FECACommandResult FECACommand_AddTextBlockToWidget::Execute(const TSharedPtr<FJsonObject>& Params)
+{
+	FString WidgetPath;
+	if (!GetStringParam(Params, TEXT("widget_path"), WidgetPath))
+	{
+		return FECACommandResult::Error(TEXT("Missing required parameter: widget_path"));
+	}
+	
+	FString TextBlockName;
+	if (!GetStringParam(Params, TEXT("text_block_name"), TextBlockName))
+	{
+		return FECACommandResult::Error(TEXT("Missing required parameter: text_block_name"));
+	}
+	
+	FString Text = TEXT("Text");
+	GetStringParam(Params, TEXT("text"), Text, false);
+	
+	double FontSize = 12;
+	GetFloatParam(Params, TEXT("font_size"), FontSize, false);
+	
+	UWidgetBlueprint* WidgetBlueprint = LoadUMGWidgetBlueprintByPath(WidgetPath);
+	if (!WidgetBlueprint)
+	{
+		return FECACommandResult::Error(FString::Printf(TEXT("Widget Blueprint not found: %s"), *WidgetPath));
+	}
+	
+	UWidgetTree* WidgetTree = WidgetBlueprint->WidgetTree;
+	if (!WidgetTree || !WidgetTree->RootWidget)
+	{
+		return FECACommandResult::Error(TEXT("Widget Blueprint has no root widget"));
+	}
+	
+	// Create Text Block
+	UTextBlock* TextBlock = WidgetTree->ConstructWidget<UTextBlock>(UTextBlock::StaticClass(), FName(*TextBlockName));
+	if (!TextBlock)
+	{
+		return FECACommandResult::Error(TEXT("Failed to create Text Block"));
+	}
+	
+	// Set text
+	TextBlock->SetText(FText::FromString(Text));
+	
+	// Set color if specified
+	FLinearColor Color = GetColorFromParams(Params, TEXT("color"), FLinearColor::White);
+	TextBlock->SetColorAndOpacity(FSlateColor(Color));
+	
+	// Add to canvas panel
+	UCanvasPanel* Canvas = Cast<UCanvasPanel>(WidgetTree->RootWidget);
+	if (Canvas)
+	{
+		UCanvasPanelSlot* Slot = Canvas->AddChildToCanvas(TextBlock);
+		if (Slot)
+		{
+			FVector2D Position = GetPositionFromParams(Params);
+			FVector2D Size = GetSizeFromParams(Params, FVector2D(200, 50));
+			
+			Slot->SetPosition(Position);
+			Slot->SetSize(Size);
+			Slot->SetAutoSize(true);
+		}
+	}
+	else
+	{
+		WidgetTree->RootWidget = TextBlock;
+	}
+	
+	FBlueprintEditorUtils::MarkBlueprintAsStructurallyModified(WidgetBlueprint);
+	
+	TSharedPtr<FJsonObject> Result = MakeResult();
+	Result->SetStringField(TEXT("text_block_name"), TextBlockName);
+	Result->SetStringField(TEXT("text"), Text);
+	return FECACommandResult::Success(Result);
+}
+
+//------------------------------------------------------------------------------
+// AddButtonToWidget
+//------------------------------------------------------------------------------
+
+FECACommandResult FECACommand_AddButtonToWidget::Execute(const TSharedPtr<FJsonObject>& Params)
+{
+	FString WidgetPath;
+	if (!GetStringParam(Params, TEXT("widget_path"), WidgetPath))
+	{
+		return FECACommandResult::Error(TEXT("Missing required parameter: widget_path"));
+	}
+	
+	FString ButtonName;
+	if (!GetStringParam(Params, TEXT("button_name"), ButtonName))
+	{
+		return FECACommandResult::Error(TEXT("Missing required parameter: button_name"));
+	}
+	
+	FString ButtonText = TEXT("Button");
+	GetStringParam(Params, TEXT("text"), ButtonText, false);
+	
+	UWidgetBlueprint* WidgetBlueprint = LoadUMGWidgetBlueprintByPath(WidgetPath);
+	if (!WidgetBlueprint)
+	{
+		return FECACommandResult::Error(FString::Printf(TEXT("Widget Blueprint not found: %s"), *WidgetPath));
+	}
+	
+	UWidgetTree* WidgetTree = WidgetBlueprint->WidgetTree;
+	if (!WidgetTree || !WidgetTree->RootWidget)
+	{
+		return FECACommandResult::Error(TEXT("Widget Blueprint has no root widget"));
+	}
+	
+	// Create Button
+	UButton* Button = WidgetTree->ConstructWidget<UButton>(UButton::StaticClass(), FName(*ButtonName));
+	if (!Button)
+	{
+		return FECACommandResult::Error(TEXT("Failed to create Button"));
+	}
+	
+	// Create Text Block for button content
+	FString TextBlockName = ButtonName + TEXT("_Text");
+	UTextBlock* TextBlock = WidgetTree->ConstructWidget<UTextBlock>(UTextBlock::StaticClass(), FName(*TextBlockName));
+	if (TextBlock)
+	{
+		TextBlock->SetText(FText::FromString(ButtonText));
+		Button->AddChild(TextBlock);
+	}
+	
+	// Add to canvas panel
+	UCanvasPanel* Canvas = Cast<UCanvasPanel>(WidgetTree->RootWidget);
+	if (Canvas)
+	{
+		UCanvasPanelSlot* Slot = Canvas->AddChildToCanvas(Button);
+		if (Slot)
+		{
+			FVector2D Position = GetPositionFromParams(Params);
+			FVector2D Size = GetSizeFromParams(Params, FVector2D(200, 50));
+			
+			Slot->SetPosition(Position);
+			Slot->SetSize(Size);
+		}
+	}
+	
+	FBlueprintEditorUtils::MarkBlueprintAsStructurallyModified(WidgetBlueprint);
+	
+	TSharedPtr<FJsonObject> Result = MakeResult();
+	Result->SetStringField(TEXT("button_name"), ButtonName);
+	Result->SetStringField(TEXT("text"), ButtonText);
+	return FECACommandResult::Success(Result);
+}
+
+//------------------------------------------------------------------------------
+// AddWidgetToViewport
+//------------------------------------------------------------------------------
+
+FECACommandResult FECACommand_AddWidgetToViewport::Execute(const TSharedPtr<FJsonObject>& Params)
+{
+	FString WidgetPath;
+	if (!GetStringParam(Params, TEXT("widget_path"), WidgetPath))
+	{
+		return FECACommandResult::Error(TEXT("Missing required parameter: widget_path"));
+	}
+	
+	int32 ZOrder = 0;
+	GetIntParam(Params, TEXT("z_order"), ZOrder, false);
+	
+	UWidgetBlueprint* WidgetBlueprint = LoadUMGWidgetBlueprintByPath(WidgetPath);
+	if (!WidgetBlueprint)
+	{
+		return FECACommandResult::Error(FString::Printf(TEXT("Widget Blueprint not found: %s"), *WidgetPath));
+	}
+	
+	if (!WidgetBlueprint->GeneratedClass)
+	{
+		return FECACommandResult::Error(TEXT("Widget Blueprint has no generated class. Try compiling it first."));
+	}
+	
+	// This command only works during PIE
+	UWorld* PlayWorld = GEditor->PlayWorld;
+	if (!PlayWorld)
+	{
+		return FECACommandResult::Error(TEXT("This command requires an active Play In Editor session"));
+	}
+	
+	// Get the player controller
+	APlayerController* PC = PlayWorld->GetFirstPlayerController();
+	if (!PC)
+	{
+		return FECACommandResult::Error(TEXT("No player controller found"));
+	}
+	
+	// Create and add widget to viewport
+	TSubclassOf<UUserWidget> WidgetClass = Cast<UClass>(WidgetBlueprint->GeneratedClass);
+	if (!WidgetClass)
+	{
+		return FECACommandResult::Error(TEXT("Widget blueprint does not have a valid generated class"));
+	}
+	
+	UUserWidget* Widget = CreateWidget<UUserWidget>(PC, WidgetClass);
+	if (!Widget)
+	{
+		return FECACommandResult::Error(TEXT("Failed to create widget instance"));
+	}
+	
+	Widget->AddToViewport(ZOrder);
+	
+	TSharedPtr<FJsonObject> Result = MakeResult();
+	Result->SetStringField(TEXT("widget_path"), WidgetPath);
+	Result->SetNumberField(TEXT("z_order"), ZOrder);
+	Result->SetBoolField(TEXT("added_to_viewport"), true);
+	return FECACommandResult::Success(Result);
+}
+
+//------------------------------------------------------------------------------
+// BindWidgetEvent
+//------------------------------------------------------------------------------
+
+FECACommandResult FECACommand_BindWidgetEvent::Execute(const TSharedPtr<FJsonObject>& Params)
+{
+	FString WidgetPath;
+	if (!GetStringParam(Params, TEXT("widget_path"), WidgetPath))
+	{
+		return FECACommandResult::Error(TEXT("Missing required parameter: widget_path"));
+	}
+	
+	FString WidgetName;
+	if (!GetStringParam(Params, TEXT("widget_name"), WidgetName))
+	{
+		return FECACommandResult::Error(TEXT("Missing required parameter: widget_name"));
+	}
+	
+	FString EventName;
+	if (!GetStringParam(Params, TEXT("event_name"), EventName))
+	{
+		return FECACommandResult::Error(TEXT("Missing required parameter: event_name"));
+	}
+	
+	UWidgetBlueprint* WidgetBlueprint = LoadUMGWidgetBlueprintByPath(WidgetPath);
+	if (!WidgetBlueprint)
+	{
+		return FECACommandResult::Error(FString::Printf(TEXT("Widget Blueprint not found: %s"), *WidgetPath));
+	}
+	
+	// Find the widget in the tree
+	UWidgetTree* WidgetTree = WidgetBlueprint->WidgetTree;
+	if (!WidgetTree)
+	{
+		return FECACommandResult::Error(TEXT("Widget Blueprint has no widget tree"));
+	}
+	
+	UWidget* TargetWidget = WidgetTree->FindWidget(FName(*WidgetName));
+	if (!TargetWidget)
+	{
+		return FECACommandResult::Error(FString::Printf(TEXT("Widget not found: %s"), *WidgetName));
+	}
+	
+	// For now, just mark the widget as a variable so it can be bound in Blueprints
+	TargetWidget->bIsVariable = true;
+	
+	FBlueprintEditorUtils::MarkBlueprintAsStructurallyModified(WidgetBlueprint);
+	
+	TSharedPtr<FJsonObject> Result = MakeResult();
+	Result->SetStringField(TEXT("widget_name"), WidgetName);
+	Result->SetStringField(TEXT("event_name"), EventName);
+	Result->SetBoolField(TEXT("is_variable"), true);
+	Result->SetStringField(TEXT("note"), TEXT("Widget marked as variable. Use Blueprint editor to complete event binding."));
+	return FECACommandResult::Success(Result);
+}
+
+//------------------------------------------------------------------------------
+// SetTextBlockBinding
+//------------------------------------------------------------------------------
+
+FECACommandResult FECACommand_SetTextBlockBinding::Execute(const TSharedPtr<FJsonObject>& Params)
+{
+	FString WidgetPath;
+	if (!GetStringParam(Params, TEXT("widget_path"), WidgetPath))
+	{
+		return FECACommandResult::Error(TEXT("Missing required parameter: widget_path"));
+	}
+	
+	FString TextBlockName;
+	if (!GetStringParam(Params, TEXT("text_block_name"), TextBlockName))
+	{
+		return FECACommandResult::Error(TEXT("Missing required parameter: text_block_name"));
+	}
+	
+	FString BindingFunction;
+	if (!GetStringParam(Params, TEXT("binding_function"), BindingFunction))
+	{
+		return FECACommandResult::Error(TEXT("Missing required parameter: binding_function"));
+	}
+	
+	UWidgetBlueprint* WidgetBlueprint = LoadUMGWidgetBlueprintByPath(WidgetPath);
+	if (!WidgetBlueprint)
+	{
+		return FECACommandResult::Error(FString::Printf(TEXT("Widget Blueprint not found: %s"), *WidgetPath));
+	}
+	
+	UWidgetTree* WidgetTree = WidgetBlueprint->WidgetTree;
+	if (!WidgetTree)
+	{
+		return FECACommandResult::Error(TEXT("Widget Blueprint has no widget tree"));
+	}
+	
+	UWidget* Widget = WidgetTree->FindWidget(FName(*TextBlockName));
+	UTextBlock* TextBlock = Cast<UTextBlock>(Widget);
+	if (!TextBlock)
+	{
+		return FECACommandResult::Error(FString::Printf(TEXT("Text Block not found: %s"), *TextBlockName));
+	}
+	
+	// Mark as variable for binding
+	TextBlock->bIsVariable = true;
+	
+	FBlueprintEditorUtils::MarkBlueprintAsStructurallyModified(WidgetBlueprint);
+	
+	TSharedPtr<FJsonObject> Result = MakeResult();
+	Result->SetStringField(TEXT("text_block_name"), TextBlockName);
+	Result->SetStringField(TEXT("binding_function"), BindingFunction);
+	Result->SetBoolField(TEXT("is_variable"), true);
+	Result->SetStringField(TEXT("note"), TEXT("Text Block marked as variable. Create binding function in Blueprint editor."));
+	return FECACommandResult::Success(Result);
+}
+
+//------------------------------------------------------------------------------
+// AddImageToWidget
+//------------------------------------------------------------------------------
+
+FECACommandResult FECACommand_AddImageToWidget::Execute(const TSharedPtr<FJsonObject>& Params)
+{
+	FString WidgetPath;
+	if (!GetStringParam(Params, TEXT("widget_path"), WidgetPath))
+	{
+		return FECACommandResult::Error(TEXT("Missing required parameter: widget_path"));
+	}
+	
+	FString ImageName;
+	if (!GetStringParam(Params, TEXT("image_name"), ImageName))
+	{
+		return FECACommandResult::Error(TEXT("Missing required parameter: image_name"));
+	}
+	
+	UWidgetBlueprint* WidgetBlueprint = LoadUMGWidgetBlueprintByPath(WidgetPath);
+	if (!WidgetBlueprint)
+	{
+		return FECACommandResult::Error(FString::Printf(TEXT("Widget Blueprint not found: %s"), *WidgetPath));
+	}
+	
+	UWidgetTree* WidgetTree = WidgetBlueprint->WidgetTree;
+	if (!WidgetTree || !WidgetTree->RootWidget)
+	{
+		return FECACommandResult::Error(TEXT("Widget Blueprint has no root widget"));
+	}
+	
+	// Create Image widget
+	UImage* Image = WidgetTree->ConstructWidget<UImage>(UImage::StaticClass(), FName(*ImageName));
+	if (!Image)
+	{
+		return FECACommandResult::Error(TEXT("Failed to create Image widget"));
+	}
+	
+	// Set texture if specified
+	FString TexturePath;
+	if (GetStringParam(Params, TEXT("texture_path"), TexturePath, false))
+	{
+		UTexture2D* Texture = LoadObject<UTexture2D>(nullptr, *TexturePath);
+		if (Texture)
+		{
+			Image->SetBrushFromTexture(Texture);
+		}
+	}
+	
+	// Set color and opacity
+	FLinearColor ColorAndOpacity = GetColorFromParams(Params, TEXT("color_and_opacity"), FLinearColor::White);
+	Image->SetColorAndOpacity(ColorAndOpacity);
+	
+	// Add to canvas panel
+	UCanvasPanel* Canvas = Cast<UCanvasPanel>(WidgetTree->RootWidget);
+	if (Canvas)
+	{
+		UCanvasPanelSlot* Slot = Canvas->AddChildToCanvas(Image);
+		if (Slot)
+		{
+			FVector2D Position = GetPositionFromParams(Params);
+			FVector2D Size = GetSizeFromParams(Params, FVector2D(100, 100));
+			
+			Slot->SetPosition(Position);
+			Slot->SetSize(Size);
+		}
+	}
+	
+	FBlueprintEditorUtils::MarkBlueprintAsStructurallyModified(WidgetBlueprint);
+	
+	TSharedPtr<FJsonObject> Result = MakeResult();
+	Result->SetStringField(TEXT("image_name"), ImageName);
+	return FECACommandResult::Success(Result);
+}
+
+//------------------------------------------------------------------------------
+// GetWidgetInfo
+//------------------------------------------------------------------------------
+
+FECACommandResult FECACommand_GetWidgetInfo::Execute(const TSharedPtr<FJsonObject>& Params)
+{
+	FString WidgetPath;
+	if (!GetStringParam(Params, TEXT("widget_path"), WidgetPath))
+	{
+		return FECACommandResult::Error(TEXT("Missing required parameter: widget_path"));
+	}
+	
+	UWidgetBlueprint* WidgetBlueprint = LoadUMGWidgetBlueprintByPath(WidgetPath);
+	if (!WidgetBlueprint)
+	{
+		return FECACommandResult::Error(FString::Printf(TEXT("Widget Blueprint not found: %s"), *WidgetPath));
+	}
+	
+	TSharedPtr<FJsonObject> Result = MakeResult();
+	Result->SetStringField(TEXT("name"), WidgetBlueprint->GetName());
+	Result->SetStringField(TEXT("path"), WidgetBlueprint->GetPathName());
+	Result->SetStringField(TEXT("parent_class"), WidgetBlueprint->ParentClass ? WidgetBlueprint->ParentClass->GetName() : TEXT("None"));
+	
+	// Get widgets in tree
+	TArray<TSharedPtr<FJsonValue>> WidgetsArray;
+	UWidgetTree* WidgetTree = WidgetBlueprint->WidgetTree;
+	if (WidgetTree)
+	{
+		WidgetTree->ForEachWidget([&WidgetsArray](UWidget* Widget)
+		{
+			TSharedPtr<FJsonObject> WidgetJson = MakeShared<FJsonObject>();
+			WidgetJson->SetStringField(TEXT("name"), Widget->GetName());
+			WidgetJson->SetStringField(TEXT("class"), Widget->GetClass()->GetName());
+			WidgetJson->SetBoolField(TEXT("is_variable"), Widget->bIsVariable);
+			WidgetsArray.Add(MakeShared<FJsonValueObject>(WidgetJson));
+		});
+	}
+	Result->SetArrayField(TEXT("widgets"), WidgetsArray);
+	Result->SetNumberField(TEXT("widget_count"), WidgetsArray.Num());
+	
+	return FECACommandResult::Success(Result);
+}
diff --git a/Engine/Plugins/Experimental/ECABridge/Source/ECABridge/Private/Commands/ECAViewCommands.cpp b/Engine/Plugins/Experimental/ECABridge/Source/ECABridge/Private/Commands/ECAViewCommands.cpp
new file mode 100644
--- /dev/null
+++ b/Engine/Plugins/Experimental/ECABridge/Source/ECABridge/Private/Commands/ECAViewCommands.cpp
@@ -0,0 +1,728 @@
+// Copyright Epic Games, Inc. All Rights Reserved.
+
+#include "Commands/ECAViewCommands.h"
+#include "Commands/ECACommand.h"
+#include "Editor.h"
+#include "Engine/World.h"
+#include "GameFramework/Actor.h"
+#include "Camera/CameraActor.h"
+#include "Camera/CameraComponent.h"
+#include "Components/PrimitiveComponent.h"
+#include "Components/StaticMeshComponent.h"
+#include "Engine/StaticMeshActor.h"
+#include "EngineUtils.h"
+#include "LevelEditor.h"
+#include "IAssetViewport.h"
+#include "SLevelViewport.h"
+#include "EditorViewportClient.h"
+#include "Modules/ModuleManager.h"
+#include "ConvexVolume.h"
+
+// Register all view commands
+REGISTER_ECA_COMMAND(FECACommand_GetActorsInView)
+REGISTER_ECA_COMMAND(FECACommand_GetViewportCamera)
+REGISTER_ECA_COMMAND(FECACommand_DescribeView)
+
+//------------------------------------------------------------------------------
+// Helper structures and functions
+//------------------------------------------------------------------------------
+
+namespace ViewHelpers
+{
+	/** Information about an actor visible in the view */
+	struct FVisibleActorInfo
+	{
+		AActor* Actor;
+		float Distance;
+		FVector RelativeLocation;  // Location relative to camera (Forward, Right, Up)
+		FBox WorldBounds;
+	};
+
+	/** Get the active editor viewport client using the same method as other ECA commands */
+	FEditorViewportClient* GetActiveEditorViewportClient()
+	{
+		FLevelEditorModule& LevelEditorModule = FModuleManager::GetModuleChecked<FLevelEditorModule>(TEXT("LevelEditor"));
+		TSharedPtr<IAssetViewport> ActiveViewport = LevelEditorModule.GetFirstActiveViewport();
+		
+		if (ActiveViewport.IsValid() && ActiveViewport->GetActiveViewport())
+		{
+			return static_cast<FEditorViewportClient*>(ActiveViewport->GetActiveViewport()->GetClient());
+		}
+		
+		return nullptr;
+	}
+
+	/** Get camera transform from a named camera actor or the editor viewport */
+	bool GetCameraTransform(const FString& CameraName, UWorld* World, FVector& OutLocation, FRotator& OutRotation, float& OutFOV, FString& OutSourceName)
+	{
+		// If camera name specified, find that camera actor
+		if (!CameraName.IsEmpty() && World)
+		{
+			for (TActorIterator<ACameraActor> It(World); It; ++It)
+			{
+				ACameraActor* CameraActor = *It;
+				if (CameraActor && (CameraActor->GetActorLabel().Contains(CameraName) || CameraActor->GetName().Contains(CameraName)))
+				{
+					OutLocation = CameraActor->GetActorLocation();
+					OutRotation = CameraActor->GetActorRotation();
+					OutSourceName = CameraActor->GetActorLabel();
+					
+					if (UCameraComponent* CameraComp = CameraActor->GetCameraComponent())
+					{
+						OutFOV = CameraComp->FieldOfView;
+					}
+					else
+					{
+						OutFOV = 90.0f;
+					}
+					return true;
+				}
+			}
+			return false; // Camera not found
+		}
+		
+		// Use editor viewport camera
+		FEditorViewportClient* ViewportClient = GetActiveEditorViewportClient();
+		if (ViewportClient)
+		{
+			OutLocation = ViewportClient->GetViewLocation();
+			OutRotation = ViewportClient->GetViewRotation();
+			OutFOV = ViewportClient->ViewFOV;
+			OutSourceName = TEXT("EditorViewport");
+			return true;
+		}
+		
+		return false;
+	}
+
+	/** 
+	 * Check if an actor is potentially visible from the camera.
+	 * Uses a generous check - we'd rather include actors that might be slightly out of view
+	 * than miss actors that are clearly visible.
+	 */
+	bool IsInViewFrustum(const FVector& CameraLocation, const FVector& CameraForward, const FVector& CameraRight, const FVector& CameraUp,
+						 float HalfFOVRadians, float AspectRatio, float MaxDistance, 
+						 const FVector& TestPoint, const FBox* OptionalBounds = nullptr)
+	{
+		// If we have bounds, use the closest point on bounds to camera for the primary test
+		FVector PointToTest = TestPoint;
+		if (OptionalBounds && OptionalBounds->IsValid)
+		{
+			// Use bounds center, but also we'll test if ANY part of bounds is in front
+			PointToTest = OptionalBounds->GetCenter();
+		}
+		
+		FVector ToPoint = PointToTest - CameraLocation;
+		float Distance = ToPoint.Size();
+		
+		// Distance check
+		if (Distance > MaxDistance)
+		{
+			return false;
+		}
+		
+		// Very close objects are always "in view" (within 2 meters)
+		if (Distance < 200.0f)
+		{
+			return true;
+		}
+		
+		// Check if the point (or any corner of bounds) is in front of camera
+		// "In front" means the dot product with forward is positive
+		float ForwardDot = FVector::DotProduct(ToPoint.GetSafeNormal(), CameraForward);
+		
+		// If center is behind camera, check bounds corners
+		if (ForwardDot < -0.1f)
+		{
+			if (OptionalBounds && OptionalBounds->IsValid)
+			{
+				// Check if any corner is in front
+				FVector Corners[8] = {
+					FVector(OptionalBounds->Min.X, OptionalBounds->Min.Y, OptionalBounds->Min.Z),
+					FVector(OptionalBounds->Max.X, OptionalBounds->Min.Y, OptionalBounds->Min.Z),
+					FVector(OptionalBounds->Min.X, OptionalBounds->Max.Y, OptionalBounds->Min.Z),
+					FVector(OptionalBounds->Max.X, OptionalBounds->Max.Y, OptionalBounds->Min.Z),
+					FVector(OptionalBounds->Min.X, OptionalBounds->Min.Y, OptionalBounds->Max.Z),
+					FVector(OptionalBounds->Max.X, OptionalBounds->Min.Y, OptionalBounds->Max.Z),
+					FVector(OptionalBounds->Min.X, OptionalBounds->Max.Y, OptionalBounds->Max.Z),
+					FVector(OptionalBounds->Max.X, OptionalBounds->Max.Y, OptionalBounds->Max.Z)
+				};
+				
+				bool bAnyInFront = false;
+				for (const FVector& Corner : Corners)
+				{
+					FVector ToCorner = Corner - CameraLocation;
+					if (FVector::DotProduct(ToCorner.GetSafeNormal(), CameraForward) > -0.1f)
+					{
+						bAnyInFront = true;
+						break;
+					}
+				}
+				
+				if (!bAnyInFront)
+				{
+					return false;
+				}
+			}
+			else
+			{
+				return false;
+			}
+		}
+		
+		// For the FOV check, be generous - use a wide angle that accounts for:
+		// - Horizontal FOV (wider than vertical due to aspect ratio)
+		// - Large objects that extend beyond their center point
+		// - Some margin for objects at screen edges
+		float HorizontalHalfFOV = FMath::Atan(FMath::Tan(HalfFOVRadians) * AspectRatio);
+		float MaxHalfAngle = FMath::Max(HalfFOVRadians, HorizontalHalfFOV);
+		
+		// Add generous margin (50%) to catch objects at edges and large objects
+		MaxHalfAngle *= 1.5f;
+		
+		// Also add extra angle for bounds size
+		if (OptionalBounds && OptionalBounds->IsValid && Distance > 1.0f)
+		{
+			float BoundsRadius = OptionalBounds->GetExtent().Size();
+			float BoundsAngle = FMath::Atan(BoundsRadius / Distance);
+			MaxHalfAngle += BoundsAngle;
+		}
+		
+		// Calculate angle from camera forward to the point
+		float AngleToPoint = FMath::Acos(FMath::Clamp(ForwardDot, -1.0f, 1.0f));
+		
+		return AngleToPoint <= MaxHalfAngle;
+	}
+
+	/** Get actors visible from a camera position/orientation */
+	TArray<FVisibleActorInfo> GetVisibleActors(
+		UWorld* World,
+		const FVector& CameraLocation,
+		const FRotator& CameraRotation,
+		float FOV,
+		float MinDistance,
+		float MaxDistance,
+		bool bIncludeHidden,
+		const FString& ClassFilter,
+		int32 MaxResults)
+	{
+		TArray<FVisibleActorInfo> VisibleActors;
+		
+		if (!World)
+		{
+			return VisibleActors;
+		}
+		
+		// Calculate camera vectors
+		FVector CameraForward = CameraRotation.Vector();
+		FVector CameraRight = FRotationMatrix(CameraRotation).GetUnitAxis(EAxis::Y);
+		FVector CameraUp = FRotationMatrix(CameraRotation).GetUnitAxis(EAxis::Z);
+		
+		float HalfFOVRadians = FMath::DegreesToRadians(FOV * 0.5f);
+		float AspectRatio = 16.0f / 9.0f;
+		
+		for (TActorIterator<AActor> It(World); It; ++It)
+		{
+			AActor* Actor = *It;
+			if (!Actor || Actor->IsPendingKillPending())
+			{
+				continue;
+			}
+			
+			// Skip hidden actors unless requested
+			if (!bIncludeHidden && Actor->IsHiddenEd())
+			{
+				continue;
+			}
+			
+			// Skip transient actors
+			if (Actor->HasAnyFlags(RF_Transient))
+			{
+				continue;
+			}
+			
+			// Apply class filter
+			if (!ClassFilter.IsEmpty())
+			{
+				FString ActorClassName = Actor->GetClass()->GetName();
+				if (!ActorClassName.Contains(ClassFilter, ESearchCase::IgnoreCase))
+				{
+					continue;
+				}
+			}
+			
+			// Get actor location and bounds
+			FVector ActorLocation = Actor->GetActorLocation();
+			FBox ActorBounds = Actor->GetComponentsBoundingBox(true);
+			
+			// Check distance
+			float Distance = FVector::Dist(CameraLocation, ActorLocation);
+			if (Distance < MinDistance || Distance > MaxDistance)
+			{
+				continue;
+			}
+			
+			// Check if in view frustum
+			if (!IsInViewFrustum(CameraLocation, CameraForward, CameraRight, CameraUp,
+								 HalfFOVRadians, AspectRatio, MaxDistance,
+								 ActorLocation, ActorBounds.IsValid ? &ActorBounds : nullptr))
+			{
+				continue;
+			}
+			
+			// Calculate relative position
+			FVector ToActor = ActorLocation - CameraLocation;
+			FVector RelativeLocation(
+				FVector::DotProduct(ToActor, CameraForward),
+				FVector::DotProduct(ToActor, CameraRight),
+				FVector::DotProduct(ToActor, CameraUp)
+			);
+			
+			FVisibleActorInfo Info;
+			Info.Actor = Actor;
+			Info.Distance = Distance;
+			Info.RelativeLocation = RelativeLocation;
+			Info.WorldBounds = ActorBounds;
+			
+			VisibleActors.Add(Info);
+		}
+		
+		// Sort by distance (nearest first)
+		VisibleActors.Sort([](const FVisibleActorInfo& A, const FVisibleActorInfo& B)
+		{
+			return A.Distance < B.Distance;
+		});
+		
+		// Limit results
+		if (MaxResults > 0 && VisibleActors.Num() > MaxResults)
+		{
+			VisibleActors.SetNum(MaxResults);
+		}
+		
+		return VisibleActors;
+	}
+
+	/** Get a descriptive category for an actor */
+	FString GetActorCategory(AActor* Actor)
+	{
+		if (!Actor)
+		{
+			return TEXT("Unknown");
+		}
+		
+		FString ClassName = Actor->GetClass()->GetName();
+		
+		if (ClassName.Contains(TEXT("Light"))) return TEXT("Lighting");
+		if (ClassName.Contains(TEXT("Camera"))) return TEXT("Cameras");
+		if (ClassName.Contains(TEXT("StaticMesh"))) return TEXT("StaticMeshes");
+		if (ClassName.Contains(TEXT("SkeletalMesh")) || ClassName.Contains(TEXT("Character"))) return TEXT("Characters");
+		if (ClassName.Contains(TEXT("Trigger")) || ClassName.Contains(TEXT("Volume"))) return TEXT("Volumes");
+		if (ClassName.Contains(TEXT("Audio")) || ClassName.Contains(TEXT("Sound"))) return TEXT("Audio");
+		if (ClassName.Contains(TEXT("Emitter")) || ClassName.Contains(TEXT("Particle")) || ClassName.Contains(TEXT("Niagara"))) return TEXT("Effects");
+		if (ClassName.Contains(TEXT("Landscape")) || ClassName.Contains(TEXT("Terrain"))) return TEXT("Landscape");
+		if (ClassName.Contains(TEXT("Foliage"))) return TEXT("Foliage");
+		if (ClassName.Contains(TEXT("Decal"))) return TEXT("Decals");
+		if (ClassName.Contains(TEXT("Reflection")) || ClassName.Contains(TEXT("PostProcess"))) return TEXT("PostProcess");
+		if (ClassName.Contains(TEXT("BP_")) || Actor->GetClass()->ClassGeneratedBy != nullptr) return TEXT("Blueprints");
+		
+		return TEXT("Other");
+	}
+
+	/** Get screen position description */
+	FString GetScreenPosition(const FVector& RelativeLocation)
+	{
+		FString Position;
+		
+		// Horizontal position
+		if (FMath::Abs(RelativeLocation.Y) < 100.0f)
+		{
+			Position = TEXT("center");
+		}
+		else
+		{
+			Position = RelativeLocation.Y > 0 ? TEXT("right") : TEXT("left");
+		}
+		
+		// Vertical position
+		if (FMath::Abs(RelativeLocation.Z) > 100.0f)
+		{
+			Position += RelativeLocation.Z > 0 ? TEXT("-above") : TEXT("-below");
+		}
+		
+		return Position;
+	}
+}
+
+//------------------------------------------------------------------------------
+// GetActorsInView Implementation
+//------------------------------------------------------------------------------
+
+FECACommandResult FECACommand_GetActorsInView::Execute(const TSharedPtr<FJsonObject>& Params)
+{
+	UWorld* World = GetEditorWorld();
+	if (!World)
+	{
+		return FECACommandResult::Error(TEXT("No world available"));
+	}
+	
+	// Get parameters
+	FString CameraName;
+	double MaxDistance = 10000.0;  // 100 meters default
+	double MinDistance = 0.0;
+	bool bIncludeHidden = false;
+	FString ClassFilter;
+	bool bIncludeComponents = false;
+	int32 MaxResults = 100;
+	double FOVScale = 1.0;
+	
+	GetStringParam(Params, TEXT("camera_name"), CameraName, false);
+	GetFloatParam(Params, TEXT("max_distance"), MaxDistance, false);
+	GetFloatParam(Params, TEXT("min_distance"), MinDistance, false);
+	GetBoolParam(Params, TEXT("include_hidden"), bIncludeHidden, false);
+	GetStringParam(Params, TEXT("class_filter"), ClassFilter, false);
+	GetBoolParam(Params, TEXT("include_components"), bIncludeComponents, false);
+	GetIntParam(Params, TEXT("max_results"), MaxResults, false);
+	GetFloatParam(Params, TEXT("fov_scale"), FOVScale, false);
+	
+	// Get camera transform
+	FVector CameraLocation;
+	FRotator CameraRotation;
+	float CameraFOV;
+	FString SourceName;
+	
+	if (!ViewHelpers::GetCameraTransform(CameraName, World, CameraLocation, CameraRotation, CameraFOV, SourceName))
+	{
+		if (!CameraName.IsEmpty())
+		{
+			return FECACommandResult::Error(FString::Printf(TEXT("Camera not found: %s"), *CameraName));
+		}
+		return FECACommandResult::Error(TEXT("No editor viewport available. Make sure a Level Editor viewport is open."));
+	}
+	
+	// Apply FOV scale
+	float EffectiveFOV = FMath::Clamp(CameraFOV * (float)FOVScale, 10.0f, 170.0f);
+	
+	// Get visible actors
+	TArray<ViewHelpers::FVisibleActorInfo> VisibleActors = ViewHelpers::GetVisibleActors(
+		World, CameraLocation, CameraRotation, EffectiveFOV,
+		(float)MinDistance, (float)MaxDistance, bIncludeHidden, ClassFilter, MaxResults
+	);
+	
+	// Build result
+	TSharedPtr<FJsonObject> Result = MakeResult();
+	
+	// Camera info
+	TSharedPtr<FJsonObject> CameraJson = MakeShared<FJsonObject>();
+	CameraJson->SetStringField(TEXT("source"), SourceName);
+	CameraJson->SetObjectField(TEXT("location"), VectorToJson(CameraLocation));
+	CameraJson->SetObjectField(TEXT("rotation"), RotatorToJson(CameraRotation));
+	CameraJson->SetNumberField(TEXT("fov"), CameraFOV);
+	CameraJson->SetNumberField(TEXT("effective_fov"), EffectiveFOV);
+	Result->SetObjectField(TEXT("camera"), CameraJson);
+	
+	// Actors array
+	TArray<TSharedPtr<FJsonValue>> ActorsArray;
+	for (const ViewHelpers::FVisibleActorInfo& Info : VisibleActors)
+	{
+		AActor* Actor = Info.Actor;
+		
+		TSharedPtr<FJsonObject> ActorJson = MakeShared<FJsonObject>();
+		ActorJson->SetStringField(TEXT("name"), Actor->GetActorLabel());
+		ActorJson->SetStringField(TEXT("internal_name"), Actor->GetName());
+		ActorJson->SetStringField(TEXT("class"), Actor->GetClass()->GetName());
+		ActorJson->SetStringField(TEXT("category"), ViewHelpers::GetActorCategory(Actor));
+		ActorJson->SetNumberField(TEXT("distance"), Info.Distance);
+		ActorJson->SetObjectField(TEXT("location"), VectorToJson(Actor->GetActorLocation()));
+		ActorJson->SetObjectField(TEXT("rotation"), RotatorToJson(Actor->GetActorRotation()));
+		
+		// Relative position to camera
+		TSharedPtr<FJsonObject> RelativeJson = MakeShared<FJsonObject>();
+		RelativeJson->SetNumberField(TEXT("forward"), Info.RelativeLocation.X);
+		RelativeJson->SetNumberField(TEXT("right"), Info.RelativeLocation.Y);
+		RelativeJson->SetNumberField(TEXT("up"), Info.RelativeLocation.Z);
+		ActorJson->SetObjectField(TEXT("relative_to_camera"), RelativeJson);
+		
+		// Screen position hint
+		ActorJson->SetStringField(TEXT("screen_position"), ViewHelpers::GetScreenPosition(Info.RelativeLocation));
+		
+		// Bounds size
+		if (Info.WorldBounds.IsValid)
+		{
+			FVector Size = Info.WorldBounds.GetSize();
+			TSharedPtr<FJsonObject> SizeJson = MakeShared<FJsonObject>();
+			SizeJson->SetNumberField(TEXT("x"), Size.X);
+			SizeJson->SetNumberField(TEXT("y"), Size.Y);
+			SizeJson->SetNumberField(TEXT("z"), Size.Z);
+			ActorJson->SetObjectField(TEXT("bounds_size"), SizeJson);
+		}
+		
+		// Tags
+		if (Actor->Tags.Num() > 0)
+		{
+			TArray<TSharedPtr<FJsonValue>> TagsArray;
+			for (const FName& Tag : Actor->Tags)
+			{
+				TagsArray.Add(MakeShared<FJsonValueString>(Tag.ToString()));
+			}
+			ActorJson->SetArrayField(TEXT("tags"), TagsArray);
+		}
+		
+		// Folder
+		FName FolderPath = Actor->GetFolderPath();
+		if (!FolderPath.IsNone())
+		{
+			ActorJson->SetStringField(TEXT("folder"), FolderPath.ToString());
+		}
+		
+		// Components if requested
+		if (bIncludeComponents)
+		{
+			TArray<TSharedPtr<FJsonValue>> ComponentsArray;
+			TArray<UActorComponent*> Components;
+			Actor->GetComponents(Components);
+			
+			for (UActorComponent* Component : Components)
+			{
+				if (Component)
+				{
+					TSharedPtr<FJsonObject> CompJson = MakeShared<FJsonObject>();
+					CompJson->SetStringField(TEXT("name"), Component->GetName());
+					CompJson->SetStringField(TEXT("class"), Component->GetClass()->GetName());
+					ComponentsArray.Add(MakeShared<FJsonValueObject>(CompJson));
+				}
+			}
+			ActorJson->SetArrayField(TEXT("components"), ComponentsArray);
+		}
+		
+		ActorsArray.Add(MakeShared<FJsonValueObject>(ActorJson));
+	}
+	
+	Result->SetArrayField(TEXT("actors"), ActorsArray);
+	Result->SetNumberField(TEXT("count"), ActorsArray.Num());
+	Result->SetNumberField(TEXT("max_distance"), MaxDistance);
+	Result->SetNumberField(TEXT("min_distance"), MinDistance);
+	
+	return FECACommandResult::Success(Result);
+}
+
+//------------------------------------------------------------------------------
+// GetViewportCamera Implementation
+//------------------------------------------------------------------------------
+
+FECACommandResult FECACommand_GetViewportCamera::Execute(const TSharedPtr<FJsonObject>& Params)
+{
+	int32 ViewportIndex = 0;
+	GetIntParam(Params, TEXT("viewport_index"), ViewportIndex, false);
+	
+	FEditorViewportClient* ViewportClient = ViewHelpers::GetActiveEditorViewportClient();
+	if (!ViewportClient)
+	{
+		return FECACommandResult::Error(TEXT("No active editor viewport found. Make sure a Level Editor viewport is open."));
+	}
+	
+	TSharedPtr<FJsonObject> Result = MakeResult();
+	
+	// Location
+	FVector Location = ViewportClient->GetViewLocation();
+	Result->SetObjectField(TEXT("location"), VectorToJson(Location));
+	
+	// Rotation
+	FRotator Rotation = ViewportClient->GetViewRotation();
+	Result->SetObjectField(TEXT("rotation"), RotatorToJson(Rotation));
+	
+	// Direction vectors
+	FVector Forward = Rotation.Vector();
+	FVector Right = FRotationMatrix(Rotation).GetUnitAxis(EAxis::Y);
+	FVector Up = FRotationMatrix(Rotation).GetUnitAxis(EAxis::Z);
+	
+	TSharedPtr<FJsonObject> DirectionsJson = MakeShared<FJsonObject>();
+	DirectionsJson->SetObjectField(TEXT("forward"), VectorToJson(Forward));
+	DirectionsJson->SetObjectField(TEXT("right"), VectorToJson(Right));
+	DirectionsJson->SetObjectField(TEXT("up"), VectorToJson(Up));
+	Result->SetObjectField(TEXT("directions"), DirectionsJson);
+	
+	// Camera properties
+	Result->SetNumberField(TEXT("fov"), ViewportClient->ViewFOV);
+	Result->SetBoolField(TEXT("is_perspective"), ViewportClient->IsPerspective());
+	Result->SetBoolField(TEXT("is_realtime"), ViewportClient->IsRealtime());
+	
+	// View mode
+	FString ViewModeStr;
+	EViewModeIndex ViewMode = ViewportClient->GetViewMode();
+	switch (ViewMode)
+	{
+		case VMI_Lit: ViewModeStr = TEXT("Lit"); break;
+		case VMI_Unlit: ViewModeStr = TEXT("Unlit"); break;
+		case VMI_Wireframe: ViewModeStr = TEXT("Wireframe"); break;
+		case VMI_BrushWireframe: ViewModeStr = TEXT("BrushWireframe"); break;
+		case VMI_LightingOnly: ViewModeStr = TEXT("LightingOnly"); break;
+		default: ViewModeStr = TEXT("Other"); break;
+	}
+	Result->SetStringField(TEXT("view_mode"), ViewModeStr);
+	
+	return FECACommandResult::Success(Result);
+}
+
+//------------------------------------------------------------------------------
+// DescribeView Implementation
+//------------------------------------------------------------------------------
+
+FECACommandResult FECACommand_DescribeView::Execute(const TSharedPtr<FJsonObject>& Params)
+{
+	UWorld* World = GetEditorWorld();
+	if (!World)
+	{
+		return FECACommandResult::Error(TEXT("No world available"));
+	}
+	
+	// Get parameters
+	FString CameraName;
+	double MaxDistance = 50000.0;  // 500 meters
+	double NearThreshold = 500.0;  // 5 meters
+	double MidThreshold = 2000.0;  // 20 meters
+	bool bIncludeHidden = false;
+	
+	GetStringParam(Params, TEXT("camera_name"), CameraName, false);
+	GetFloatParam(Params, TEXT("max_distance"), MaxDistance, false);
+	GetFloatParam(Params, TEXT("near_threshold"), NearThreshold, false);
+	GetFloatParam(Params, TEXT("mid_threshold"), MidThreshold, false);
+	GetBoolParam(Params, TEXT("include_hidden"), bIncludeHidden, false);
+	
+	// Get camera transform
+	FVector CameraLocation;
+	FRotator CameraRotation;
+	float CameraFOV;
+	FString SourceName;
+	
+	if (!ViewHelpers::GetCameraTransform(CameraName, World, CameraLocation, CameraRotation, CameraFOV, SourceName))
+	{
+		if (!CameraName.IsEmpty())
+		{
+			return FECACommandResult::Error(FString::Printf(TEXT("Camera not found: %s"), *CameraName));
+		}
+		return FECACommandResult::Error(TEXT("No editor viewport available. Make sure a Level Editor viewport is open."));
+	}
+	
+	// Get all visible actors (no max limit)
+	TArray<ViewHelpers::FVisibleActorInfo> VisibleActors = ViewHelpers::GetVisibleActors(
+		World, CameraLocation, CameraRotation, CameraFOV,
+		0.0f, (float)MaxDistance, bIncludeHidden, TEXT(""), 0
+	);
+	
+	// Categorize by distance zones
+	TMap<FString, TArray<TPair<FString, float>>> NearByCategory;  // Category -> [(ActorName, Distance)]
+	TMap<FString, TArray<TPair<FString, float>>> MidByCategory;
+	TMap<FString, TArray<TPair<FString, float>>> FarByCategory;
+	
+	int32 NearCount = 0, MidCount = 0, FarCount = 0;
+	
+	for (const ViewHelpers::FVisibleActorInfo& Info : VisibleActors)
+	{
+		FString Category = ViewHelpers::GetActorCategory(Info.Actor);
+		FString ActorName = Info.Actor->GetActorLabel();
+		
+		if (Info.Distance <= NearThreshold)
+		{
+			NearByCategory.FindOrAdd(Category).Add(TPair<FString, float>(ActorName, Info.Distance));
+			NearCount++;
+		}
+		else if (Info.Distance <= MidThreshold)
+		{
+			MidByCategory.FindOrAdd(Category).Add(TPair<FString, float>(ActorName, Info.Distance));
+			MidCount++;
+		}
+		else
+		{
+			FarByCategory.FindOrAdd(Category).Add(TPair<FString, float>(ActorName, Info.Distance));
+			FarCount++;
+		}
+	}
+	
+	// Helper to build zone JSON
+	auto BuildZoneJson = [](const TMap<FString, TArray<TPair<FString, float>>>& ByCategory, int32 TotalCount) -> TSharedPtr<FJsonObject>
+	{
+		TSharedPtr<FJsonObject> ZoneJson = MakeShared<FJsonObject>();
+		ZoneJson->SetNumberField(TEXT("total_count"), TotalCount);
+		
+		TArray<TSharedPtr<FJsonValue>> CategoriesArray;
+		for (const auto& Pair : ByCategory)
+		{
+			TSharedPtr<FJsonObject> CatJson = MakeShared<FJsonObject>();
+			CatJson->SetStringField(TEXT("category"), Pair.Key);
+			CatJson->SetNumberField(TEXT("count"), Pair.Value.Num());
+			
+			// List actors (limit to 10 per category)
+			TArray<TSharedPtr<FJsonValue>> ActorsList;
+			int32 Count = 0;
+			for (const auto& ActorPair : Pair.Value)
+			{
+				if (Count >= 10) break;
+				TSharedPtr<FJsonObject> ActorJson = MakeShared<FJsonObject>();
+				ActorJson->SetStringField(TEXT("name"), ActorPair.Key);
+				ActorJson->SetNumberField(TEXT("distance"), FMath::RoundToInt(ActorPair.Value));
+				ActorsList.Add(MakeShared<FJsonValueObject>(ActorJson));
+				Count++;
+			}
+			CatJson->SetArrayField(TEXT("actors"), ActorsList);
+			
+			CategoriesArray.Add(MakeShared<FJsonValueObject>(CatJson));
+		}
+		ZoneJson->SetArrayField(TEXT("categories"), CategoriesArray);
+		
+		return ZoneJson;
+	};
+	
+	// Build result
+	TSharedPtr<FJsonObject> Result = MakeResult();
+	
+	// Camera info
+	TSharedPtr<FJsonObject> CameraJson = MakeShared<FJsonObject>();
+	CameraJson->SetStringField(TEXT("source"), SourceName);
+	CameraJson->SetObjectField(TEXT("location"), VectorToJson(CameraLocation));
+	CameraJson->SetObjectField(TEXT("rotation"), RotatorToJson(CameraRotation));
+	CameraJson->SetNumberField(TEXT("fov"), CameraFOV);
+	
+	// Looking direction description
+	FVector Forward = CameraRotation.Vector();
+	FString LookingDir;
+	if (FMath::Abs(Forward.Z) > 0.7f)
+	{
+		LookingDir = Forward.Z > 0 ? TEXT("up") : TEXT("down");
+	}
+	else
+	{
+		float Yaw = CameraRotation.Yaw;
+		while (Yaw < 0) Yaw += 360.f;
+		while (Yaw >= 360) Yaw -= 360.f;
+		
+		if (Yaw < 45.f || Yaw >= 315.f) LookingDir = TEXT("North (+X)");
+		else if (Yaw < 135.f) LookingDir = TEXT("East (+Y)");
+		else if (Yaw < 225.f) LookingDir = TEXT("South (-X)");
+		else LookingDir = TEXT("West (-Y)");
+	}
+	CameraJson->SetStringField(TEXT("looking_direction"), LookingDir);
+	Result->SetObjectField(TEXT("camera"), CameraJson);
+	
+	// Zones
+	TSharedPtr<FJsonObject> NearJson = BuildZoneJson(NearByCategory, NearCount);
+	NearJson->SetStringField(TEXT("description"), FString::Printf(TEXT("Within %.1f meters"), NearThreshold / 100.0));
+	Result->SetObjectField(TEXT("near"), NearJson);
+	
+	TSharedPtr<FJsonObject> MidJson = BuildZoneJson(MidByCategory, MidCount);
+	MidJson->SetStringField(TEXT("description"), FString::Printf(TEXT("%.1f to %.1f meters"), NearThreshold / 100.0, MidThreshold / 100.0));
+	Result->SetObjectField(TEXT("mid"), MidJson);
+	
+	TSharedPtr<FJsonObject> FarJson = BuildZoneJson(FarByCategory, FarCount);
+	FarJson->SetStringField(TEXT("description"), FString::Printf(TEXT("%.1f to %.1f meters"), MidThreshold / 100.0, MaxDistance / 100.0));
+	Result->SetObjectField(TEXT("far"), FarJson);
+	
+	// Summary
+	TSharedPtr<FJsonObject> SummaryJson = MakeShared<FJsonObject>();
+	SummaryJson->SetNumberField(TEXT("total_visible"), VisibleActors.Num());
+	SummaryJson->SetNumberField(TEXT("near_count"), NearCount);
+	SummaryJson->SetNumberField(TEXT("mid_count"), MidCount);
+	SummaryJson->SetNumberField(TEXT("far_count"), FarCount);
+	Result->SetObjectField(TEXT("summary"), SummaryJson);
+	
+	return FECACommandResult::Success(Result);
+}
