ECABridge experimental MCP plugin

diff --git a/Engine/Plugins/Experimental/ECABridge/Source/ECABridge/Private/Commands/ECAWidgetTreeCommands.cpp b/Engine/Plugins/Experimental/ECABridge/Source/ECABridge/Private/Commands/ECAWidgetTreeCommands.cpp
new file mode 100644
--- /dev/null
+++ b/Engine/Plugins/Experimental/ECABridge/Source/ECABridge/Private/Commands/ECAWidgetTreeCommands.cpp
@@ -0,0 +1,1342 @@
+// Copyright Epic Games, Inc. All Rights Reserved.
+
+#include "Commands/ECAWidgetTreeCommands.h"
+#include "Commands/ECACommand.h"
+
+// Engine - Widget Blueprint
+#include "WidgetBlueprint.h"
+#include "Blueprint/WidgetTree.h"
+#include "Blueprint/UserWidget.h"
+
+// Engine - UMG Components
+#include "Components/Widget.h"
+#include "Components/PanelWidget.h"
+#include "Components/ContentWidget.h"
+#include "Components/SizeBox.h"
+#include "Components/Overlay.h"
+#include "Components/Border.h"
+#include "Components/ScaleBox.h"
+#include "Components/CanvasPanel.h"
+#include "Components/VerticalBox.h"
+#include "Components/HorizontalBox.h"
+#include "Components/StackBox.h"
+#include "Components/TextBlock.h"
+#include "Components/Image.h"
+
+// Engine - UMG Slots (SlateWrapperTypes for FSlateChildSize; individual slot headers
+// are NOT needed since we use UE reflection to set slot properties uniformly)
+#include "Components/PanelSlot.h"
+#include "Components/SlateWrapperTypes.h"
+
+// Engine - Blueprint utilities
+#include "Kismet2/BlueprintEditorUtils.h"
+#include "UObject/UObjectIterator.h"
+#include "UObject/SavePackage.h"
+
+// Register all widget tree commands
+REGISTER_ECA_COMMAND(FECACommand_WrapWidget)
+REGISTER_ECA_COMMAND(FECACommand_UnwrapWidget)
+REGISTER_ECA_COMMAND(FECACommand_ReparentWidget)
+REGISTER_ECA_COMMAND(FECACommand_SetSlotProperty)
+REGISTER_ECA_COMMAND(FECACommand_SetWidgetProperty)
+REGISTER_ECA_COMMAND(FECACommand_AddChildWidget)
+REGISTER_ECA_COMMAND(FECACommand_CompileWidgetBlueprint)
+REGISTER_ECA_COMMAND(FECACommand_RemoveWidget)
+
+//------------------------------------------------------------------------------
+// Shared Helper Functions
+//------------------------------------------------------------------------------
+
+static UWidgetBlueprint* LoadWidgetBlueprintFromPath(const FString& WidgetPath)
+{
+	// Try loading directly
+	UWidgetBlueprint* WBP = LoadObject<UWidgetBlueprint>(nullptr, *WidgetPath);
+	if (WBP)
+	{
+		return WBP;
+	}
+
+	// Try with common prefix patterns for short paths
+	if (!WidgetPath.StartsWith(TEXT("/")))
+	{
+		WBP = LoadObject<UWidgetBlueprint>(nullptr, *FString::Printf(TEXT("/Game/%s"), *WidgetPath));
+		if (WBP)
+		{
+			return WBP;
+		}
+	}
+
+	return nullptr;
+}
+
+static EHorizontalAlignment ParseHorizontalAlignment(const FString& Value)
+{
+	if (Value.Equals(TEXT("Left"), ESearchCase::IgnoreCase)) return HAlign_Left;
+	if (Value.Equals(TEXT("Center"), ESearchCase::IgnoreCase)) return HAlign_Center;
+	if (Value.Equals(TEXT("Right"), ESearchCase::IgnoreCase)) return HAlign_Right;
+	if (Value.Equals(TEXT("Fill"), ESearchCase::IgnoreCase)) return HAlign_Fill;
+	return HAlign_Fill; // Default
+}
+
+static EVerticalAlignment ParseVerticalAlignment(const FString& Value)
+{
+	if (Value.Equals(TEXT("Top"), ESearchCase::IgnoreCase)) return VAlign_Top;
+	if (Value.Equals(TEXT("Center"), ESearchCase::IgnoreCase)) return VAlign_Center;
+	if (Value.Equals(TEXT("Bottom"), ESearchCase::IgnoreCase)) return VAlign_Bottom;
+	if (Value.Equals(TEXT("Fill"), ESearchCase::IgnoreCase)) return VAlign_Fill;
+	return VAlign_Fill; // Default
+}
+
+static ESlateVisibility ParseVisibility(const FString& Value)
+{
+	if (Value.Equals(TEXT("Visible"), ESearchCase::IgnoreCase)) return ESlateVisibility::Visible;
+	if (Value.Equals(TEXT("Collapsed"), ESearchCase::IgnoreCase)) return ESlateVisibility::Collapsed;
+	if (Value.Equals(TEXT("Hidden"), ESearchCase::IgnoreCase)) return ESlateVisibility::Hidden;
+	if (Value.Equals(TEXT("HitTestInvisible"), ESearchCase::IgnoreCase)) return ESlateVisibility::HitTestInvisible;
+	if (Value.Equals(TEXT("SelfHitTestInvisible"), ESearchCase::IgnoreCase)) return ESlateVisibility::SelfHitTestInvisible;
+	return ESlateVisibility::Visible; // Default
+}
+
+static FMargin ParsePadding(const TSharedPtr<FJsonObject>& PaddingObj)
+{
+	FMargin Padding(0);
+	if (PaddingObj.IsValid())
+	{
+		double Left = 0, Top = 0, Right = 0, Bottom = 0;
+		PaddingObj->TryGetNumberField(TEXT("left"), Left);
+		PaddingObj->TryGetNumberField(TEXT("top"), Top);
+		PaddingObj->TryGetNumberField(TEXT("right"), Right);
+		PaddingObj->TryGetNumberField(TEXT("bottom"), Bottom);
+		Padding = FMargin(Left, Top, Right, Bottom);
+	}
+	return Padding;
+}
+
+/**
+ * Ensure all widgets in the tree have GUIDs registered with the WidgetBlueprint.
+ * The compiler's ValidateAndFixUpVariableGuids will ensure-fail on any widget
+ * that was added to the tree without a corresponding GUID entry.
+ * Call this before MarkBlueprintAsStructurallyModified.
+ */
+static void EnsureAllWidgetGuids(UWidgetBlueprint* WidgetBlueprint)
+{
+	if (!WidgetBlueprint || !WidgetBlueprint->WidgetTree)
+	{
+		return;
+	}
+
+	WidgetBlueprint->WidgetTree->ForEachWidget([WidgetBlueprint](UWidget* Widget)
+	{
+		if (Widget && !WidgetBlueprint->WidgetVariableNameToGuidMap.Contains(Widget->GetFName()))
+		{
+			WidgetBlueprint->OnVariableAdded(Widget->GetFName());
+		}
+	});
+}
+
+/**
+ * Find a UClass for any UWidget subclass by name.
+ * Searches engine UMG classes, then project classes via UClass iterator.
+ */
+static UClass* FindWidgetClass(const FString& ClassName)
+{
+	// Try common engine classes first (fast path)
+	static const TMap<FString, UClass*> EngineWidgets = []() {
+		TMap<FString, UClass*> Map;
+		Map.Add(TEXT("TextBlock"), UTextBlock::StaticClass());
+		Map.Add(TEXT("Image"), UImage::StaticClass());
+		Map.Add(TEXT("SizeBox"), USizeBox::StaticClass());
+		Map.Add(TEXT("Overlay"), UOverlay::StaticClass());
+		Map.Add(TEXT("Border"), UBorder::StaticClass());
+		Map.Add(TEXT("ScaleBox"), UScaleBox::StaticClass());
+		Map.Add(TEXT("CanvasPanel"), UCanvasPanel::StaticClass());
+		Map.Add(TEXT("VerticalBox"), UVerticalBox::StaticClass());
+		Map.Add(TEXT("HorizontalBox"), UHorizontalBox::StaticClass());
+		Map.Add(TEXT("StackBox"), UStackBox::StaticClass());
+		return Map;
+	}();
+
+	// Case-insensitive engine class lookup
+	for (const auto& Pair : EngineWidgets)
+	{
+		if (Pair.Key.Equals(ClassName, ESearchCase::IgnoreCase))
+		{
+			return Pair.Value;
+		}
+	}
+
+	// Strip 'U' prefix if present for engine lookup
+	FString SearchName = ClassName;
+	if (SearchName.StartsWith(TEXT("U")))
+	{
+		SearchName = SearchName.RightChop(1);
+		for (const auto& Pair : EngineWidgets)
+		{
+			if (Pair.Key.Equals(SearchName, ESearchCase::IgnoreCase))
+			{
+				return Pair.Value;
+			}
+		}
+	}
+
+	// Search all loaded UWidget subclasses (covers project/plugin types like CommonTextBlock, WBP_*_C)
+	for (TObjectIterator<UClass> It; It; ++It)
+	{
+		UClass* Class = *It;
+		if (Class->IsChildOf(UWidget::StaticClass()) && !Class->HasAnyClassFlags(CLASS_Abstract))
+		{
+			if (Class->GetName().Equals(ClassName, ESearchCase::IgnoreCase))
+			{
+				return Class;
+			}
+		}
+	}
+
+	return nullptr;
+}
+
+/**
+ * Find a UClass for a panel widget by name.
+ * Searches common UMG panel/content widget classes.
+ */
+static UClass* FindPanelWidgetClass(const FString& ClassName)
+{
+	// Direct lookup table for common wrapper classes
+	static const TMap<FString, UClass*> ClassMap = []() {
+		TMap<FString, UClass*> Map;
+		Map.Add(TEXT("SizeBox"), USizeBox::StaticClass());
+		Map.Add(TEXT("Overlay"), UOverlay::StaticClass());
+		Map.Add(TEXT("Border"), UBorder::StaticClass());
+		Map.Add(TEXT("ScaleBox"), UScaleBox::StaticClass());
+		Map.Add(TEXT("CanvasPanel"), UCanvasPanel::StaticClass());
+		Map.Add(TEXT("VerticalBox"), UVerticalBox::StaticClass());
+		Map.Add(TEXT("HorizontalBox"), UHorizontalBox::StaticClass());
+		Map.Add(TEXT("StackBox"), UStackBox::StaticClass());
+		return Map;
+	}();
+
+	// Case-insensitive search
+	for (const auto& Pair : ClassMap)
+	{
+		if (Pair.Key.Equals(ClassName, ESearchCase::IgnoreCase))
+		{
+			return Pair.Value;
+		}
+	}
+
+	// Try with 'U' prefix stripped/added
+	FString SearchName = ClassName;
+	if (SearchName.StartsWith(TEXT("U")))
+	{
+		SearchName = SearchName.RightChop(1);
+	}
+	for (const auto& Pair : ClassMap)
+	{
+		if (Pair.Key.Equals(SearchName, ESearchCase::IgnoreCase))
+		{
+			return Pair.Value;
+		}
+	}
+
+	// Fallback: try FindObject for any UPanelWidget subclass
+	UClass* FoundClass = FindObject<UClass>(nullptr, *FString::Printf(TEXT("/Script/UMG.%s"), *ClassName));
+	if (FoundClass && FoundClass->IsChildOf(UPanelWidget::StaticClass()))
+	{
+		return FoundClass;
+	}
+
+	return nullptr;
+}
+
+//------------------------------------------------------------------------------
+// WrapWidget
+//------------------------------------------------------------------------------
+
+FECACommandResult FECACommand_WrapWidget::Execute(const TSharedPtr<FJsonObject>& Params)
+{
+	// Parse required parameters
+	FString WidgetPath;
+	if (!GetStringParam(Params, TEXT("widget_path"), WidgetPath))
+	{
+		return FECACommandResult::Error(TEXT("Missing required parameter: widget_path"));
+	}
+
+	FString WidgetName;
+	if (!GetStringParam(Params, TEXT("widget_name"), WidgetName))
+	{
+		return FECACommandResult::Error(TEXT("Missing required parameter: widget_name"));
+	}
+
+	FString WrapperClassName;
+	if (!GetStringParam(Params, TEXT("wrapper_class"), WrapperClassName))
+	{
+		return FECACommandResult::Error(TEXT("Missing required parameter: wrapper_class"));
+	}
+
+	FString WrapperName;
+	GetStringParam(Params, TEXT("wrapper_name"), WrapperName, false);
+
+	// Load the Widget Blueprint
+	UWidgetBlueprint* WidgetBlueprint = LoadWidgetBlueprintFromPath(WidgetPath);
+	if (!WidgetBlueprint)
+	{
+		return FECACommandResult::Error(FString::Printf(TEXT("Widget Blueprint not found: %s"), *WidgetPath));
+	}
+
+	UWidgetTree* WidgetTree = WidgetBlueprint->WidgetTree;
+	if (!WidgetTree)
+	{
+		return FECACommandResult::Error(TEXT("Widget Blueprint has no widget tree"));
+	}
+
+	// Find the target widget
+	UWidget* TargetWidget = WidgetTree->FindWidget(FName(*WidgetName));
+	if (!TargetWidget)
+	{
+		return FECACommandResult::Error(FString::Printf(TEXT("Widget not found: %s"), *WidgetName));
+	}
+
+	// Find the wrapper class
+	UClass* WrapperClass = FindPanelWidgetClass(WrapperClassName);
+	if (!WrapperClass)
+	{
+		return FECACommandResult::Error(FString::Printf(TEXT("Widget class not found: %s. Valid classes: SizeBox, Overlay, Border, ScaleBox, CanvasPanel, VerticalBox, HorizontalBox, StackBox"), *WrapperClassName));
+	}
+
+	// Auto-generate wrapper name if not provided
+	if (WrapperName.IsEmpty())
+	{
+		WrapperName = FString::Printf(TEXT("%s_%s"), *WrapperClass->GetName(), *WidgetName);
+	}
+
+	// Check for name collision
+	if (WidgetTree->FindWidget(FName(*WrapperName)))
+	{
+		return FECACommandResult::Error(FString::Printf(TEXT("A widget named '%s' already exists in the tree. Provide a unique wrapper_name."), *WrapperName));
+	}
+
+	// Find the target's parent and index
+	int32 ChildIndex = INDEX_NONE;
+	UPanelWidget* ParentWidget = WidgetTree->FindWidgetParent(TargetWidget, ChildIndex);
+
+	// Construct the wrapper widget
+	UPanelWidget* WrapperWidget = WidgetTree->ConstructWidget<UPanelWidget>(WrapperClass, FName(*WrapperName));
+	if (!WrapperWidget)
+	{
+		return FECACommandResult::Error(FString::Printf(TEXT("Failed to construct wrapper widget of class: %s"), *WrapperClassName));
+	}
+
+	WrapperWidget->SetFlags(RF_Transactional);
+
+	// Perform the wrap operation
+	FString ParentName;
+	if (ParentWidget)
+	{
+		// Case 1: Widget has a parent - replace the widget at its index with the wrapper,
+		// then add the widget as a child of the wrapper
+		ParentName = ParentWidget->GetName();
+		ParentWidget->ReplaceChildAt(ChildIndex, WrapperWidget);
+		WrapperWidget->AddChild(TargetWidget);
+	}
+	else if (TargetWidget == WidgetTree->RootWidget)
+	{
+		// Case 2: Widget is the root - make the wrapper the new root
+		ParentName = TEXT("[root]");
+		WidgetTree->RootWidget = WrapperWidget;
+		WrapperWidget->AddChild(TargetWidget);
+	}
+	else
+	{
+		// Clean up the orphaned wrapper before returning
+		WidgetTree->RemoveWidget(WrapperWidget);
+		return FECACommandResult::Error(TEXT("Widget has no parent and is not the root widget. Cannot determine where to insert wrapper."));
+	}
+
+	// Ensure all widgets have GUIDs, then mark blueprint as modified (dirty + recompile)
+	EnsureAllWidgetGuids(WidgetBlueprint);
+	FBlueprintEditorUtils::MarkBlueprintAsStructurallyModified(WidgetBlueprint);
+
+	// Build result
+	TSharedPtr<FJsonObject> Result = MakeResult();
+	Result->SetStringField(TEXT("wrapper_name"), WrapperName);
+	Result->SetStringField(TEXT("wrapper_class"), WrapperClass->GetName());
+	Result->SetStringField(TEXT("wrapped_widget"), WidgetName);
+	Result->SetStringField(TEXT("parent"), ParentName);
+	return FECACommandResult::Success(Result);
+}
+
+//------------------------------------------------------------------------------
+// UnwrapWidget
+//------------------------------------------------------------------------------
+
+FECACommandResult FECACommand_UnwrapWidget::Execute(const TSharedPtr<FJsonObject>& Params)
+{
+	// Parse required parameters
+	FString WidgetPath;
+	if (!GetStringParam(Params, TEXT("widget_path"), WidgetPath))
+	{
+		return FECACommandResult::Error(TEXT("Missing required parameter: widget_path"));
+	}
+
+	FString WidgetName;
+	if (!GetStringParam(Params, TEXT("widget_name"), WidgetName))
+	{
+		return FECACommandResult::Error(TEXT("Missing required parameter: widget_name"));
+	}
+
+	// Load the Widget Blueprint
+	UWidgetBlueprint* WidgetBlueprint = LoadWidgetBlueprintFromPath(WidgetPath);
+	if (!WidgetBlueprint)
+	{
+		return FECACommandResult::Error(FString::Printf(TEXT("Widget Blueprint not found: %s"), *WidgetPath));
+	}
+
+	UWidgetTree* WidgetTree = WidgetBlueprint->WidgetTree;
+	if (!WidgetTree)
+	{
+		return FECACommandResult::Error(TEXT("Widget Blueprint has no widget tree"));
+	}
+
+	// Find the wrapper widget
+	UWidget* WrapperAsWidget = WidgetTree->FindWidget(FName(*WidgetName));
+	if (!WrapperAsWidget)
+	{
+		return FECACommandResult::Error(FString::Printf(TEXT("Widget not found: %s"), *WidgetName));
+	}
+
+	UPanelWidget* WrapperWidget = Cast<UPanelWidget>(WrapperAsWidget);
+	if (!WrapperWidget)
+	{
+		return FECACommandResult::Error(FString::Printf(TEXT("Widget '%s' is not a panel widget and cannot be unwrapped"), *WidgetName));
+	}
+
+	// Handle based on child count
+	int32 ChildCount = WrapperWidget->GetChildrenCount();
+	if (ChildCount > 1)
+	{
+		return FECACommandResult::Error(FString::Printf(TEXT("Widget '%s' has %d children. Can only unwrap widgets with 0 or 1 children."), *WidgetName, ChildCount));
+	}
+
+	// Find the wrapper's parent
+	int32 WrapperIndex = INDEX_NONE;
+	UPanelWidget* ParentWidget = WidgetTree->FindWidgetParent(WrapperAsWidget, WrapperIndex);
+
+	FString ParentName;
+	FString ChildName = TEXT("[none]");
+
+	if (ChildCount == 1)
+	{
+		// Has one child — promote it to take the wrapper's place
+		UWidget* ChildWidget = WrapperWidget->GetChildAt(0);
+		if (!ChildWidget)
+		{
+			return FECACommandResult::Error(TEXT("Child widget is null"));
+		}
+
+		ChildName = ChildWidget->GetName();
+		ChildWidget->RemoveFromParent();
+
+		if (ParentWidget)
+		{
+			ParentName = ParentWidget->GetName();
+			ParentWidget->ReplaceChildAt(WrapperIndex, ChildWidget);
+		}
+		else if (WrapperAsWidget == WidgetTree->RootWidget)
+		{
+			ParentName = TEXT("[root]");
+			WidgetTree->RootWidget = ChildWidget;
+		}
+		else
+		{
+			WrapperWidget->AddChild(ChildWidget);
+			return FECACommandResult::Error(TEXT("Wrapper has no parent and is not the root widget. Cannot determine where to place the promoted child."));
+		}
+	}
+	else
+	{
+		// Empty panel (0 children) — just remove it
+		if (ParentWidget)
+		{
+			ParentName = ParentWidget->GetName();
+			ParentWidget->RemoveChild(WrapperAsWidget);
+		}
+		else if (WrapperAsWidget == WidgetTree->RootWidget)
+		{
+			return FECACommandResult::Error(TEXT("Cannot remove the root widget when it has no children. Add a child first, or use wrap_widget to replace the root."));
+		}
+		else
+		{
+			return FECACommandResult::Error(TEXT("Wrapper has no parent and is not the root widget."));
+		}
+	}
+
+	// Remove the wrapper from the tree
+	WidgetTree->RemoveWidget(WrapperAsWidget);
+
+	// Ensure all widgets have GUIDs, then mark blueprint as modified
+	EnsureAllWidgetGuids(WidgetBlueprint);
+	FBlueprintEditorUtils::MarkBlueprintAsStructurallyModified(WidgetBlueprint);
+
+	// Build result
+	TSharedPtr<FJsonObject> Result = MakeResult();
+	Result->SetStringField(TEXT("removed_wrapper"), WidgetName);
+	Result->SetStringField(TEXT("promoted_child"), ChildName);
+	Result->SetStringField(TEXT("new_parent"), ParentName);
+	return FECACommandResult::Success(Result);
+}
+
+//------------------------------------------------------------------------------
+// ReparentWidget
+//------------------------------------------------------------------------------
+
+FECACommandResult FECACommand_ReparentWidget::Execute(const TSharedPtr<FJsonObject>& Params)
+{
+	// Parse required parameters
+	FString WidgetPath;
+	if (!GetStringParam(Params, TEXT("widget_path"), WidgetPath))
+	{
+		return FECACommandResult::Error(TEXT("Missing required parameter: widget_path"));
+	}
+
+	FString WidgetName;
+	if (!GetStringParam(Params, TEXT("widget_name"), WidgetName))
+	{
+		return FECACommandResult::Error(TEXT("Missing required parameter: widget_name"));
+	}
+
+	FString NewParentName;
+	if (!GetStringParam(Params, TEXT("new_parent_name"), NewParentName))
+	{
+		return FECACommandResult::Error(TEXT("Missing required parameter: new_parent_name"));
+	}
+
+	int32 InsertIndex = -1;
+	GetIntParam(Params, TEXT("insert_index"), InsertIndex, false);
+
+	// Load the Widget Blueprint
+	UWidgetBlueprint* WidgetBlueprint = LoadWidgetBlueprintFromPath(WidgetPath);
+	if (!WidgetBlueprint)
+	{
+		return FECACommandResult::Error(FString::Printf(TEXT("Widget Blueprint not found: %s"), *WidgetPath));
+	}
+
+	UWidgetTree* WidgetTree = WidgetBlueprint->WidgetTree;
+	if (!WidgetTree)
+	{
+		return FECACommandResult::Error(TEXT("Widget Blueprint has no widget tree"));
+	}
+
+	// Find both widgets
+	UWidget* Widget = WidgetTree->FindWidget(FName(*WidgetName));
+	if (!Widget)
+	{
+		return FECACommandResult::Error(FString::Printf(TEXT("Widget not found: %s"), *WidgetName));
+	}
+
+	UWidget* NewParentAsWidget = WidgetTree->FindWidget(FName(*NewParentName));
+	if (!NewParentAsWidget)
+	{
+		return FECACommandResult::Error(FString::Printf(TEXT("Target parent widget not found: %s"), *NewParentName));
+	}
+
+	UPanelWidget* NewParent = Cast<UPanelWidget>(NewParentAsWidget);
+	if (!NewParent)
+	{
+		return FECACommandResult::Error(FString::Printf(TEXT("Target '%s' is not a panel widget and cannot accept children"), *NewParentName));
+	}
+
+	// Self-reparent guard
+	if (Widget == NewParentAsWidget)
+	{
+		return FECACommandResult::Error(TEXT("Cannot reparent a widget to itself."));
+	}
+
+	// Root widget guard - reparenting root corrupts the tree
+	if (Widget == WidgetTree->RootWidget)
+	{
+		return FECACommandResult::Error(TEXT("Cannot reparent the root widget. Use wrap_widget to insert a new root, or restructure the hierarchy."));
+	}
+
+	// Check if the new parent can accept more children
+	if (!NewParent->CanAddMoreChildren())
+	{
+		return FECACommandResult::Error(FString::Printf(TEXT("Target '%s' cannot accept more children (it is a %s which already has a child)"), *NewParentName, *NewParent->GetClass()->GetName()));
+	}
+
+	// Circular reparent check: ensure new parent is not a descendant of the widget
+	if (UPanelWidget* WidgetAsPanel = Cast<UPanelWidget>(Widget))
+	{
+		int32 DescendantIndex = WidgetTree->FindChildIndex(WidgetAsPanel, NewParentAsWidget);
+		if (DescendantIndex != INDEX_NONE)
+		{
+			return FECACommandResult::Error(FString::Printf(TEXT("Circular reparenting detected: '%s' is a descendant of '%s'"), *NewParentName, *WidgetName));
+		}
+	}
+
+	// Capture old parent info for the response
+	int32 OldIndex = INDEX_NONE;
+	UPanelWidget* OldParent = WidgetTree->FindWidgetParent(Widget, OldIndex);
+	FString OldParentName = OldParent ? OldParent->GetName() : TEXT("[root]");
+
+	// Perform the reparent
+	// AddChild internally calls RemoveFromParent, creating a new slot
+	UPanelSlot* NewSlot = nullptr;
+	if (InsertIndex >= 0)
+	{
+		NewSlot = NewParent->InsertChildAt(InsertIndex, Widget);
+	}
+	else
+	{
+		NewSlot = NewParent->AddChild(Widget);
+	}
+
+	if (!NewSlot)
+	{
+		return FECACommandResult::Error(TEXT("Failed to add widget to new parent. The widget may have been left in a detached state."));
+	}
+
+	int32 NewIndex = NewParent->GetChildIndex(Widget);
+
+	// Ensure all widgets have GUIDs, then mark blueprint as modified
+	EnsureAllWidgetGuids(WidgetBlueprint);
+	FBlueprintEditorUtils::MarkBlueprintAsStructurallyModified(WidgetBlueprint);
+
+	// Build result
+	TSharedPtr<FJsonObject> Result = MakeResult();
+	Result->SetStringField(TEXT("widget"), WidgetName);
+	Result->SetStringField(TEXT("old_parent"), OldParentName);
+	Result->SetNumberField(TEXT("old_index"), OldIndex);
+	Result->SetStringField(TEXT("new_parent"), NewParentName);
+	Result->SetNumberField(TEXT("new_index"), NewIndex);
+	Result->SetStringField(TEXT("warning"), TEXT("Slot properties (padding, alignment, size rule) were reset to defaults. Use set_slot_property to re-apply if needed."));
+	return FECACommandResult::Success(Result);
+}
+
+//------------------------------------------------------------------------------
+// SetSlotProperty
+//------------------------------------------------------------------------------
+
+FECACommandResult FECACommand_SetSlotProperty::Execute(const TSharedPtr<FJsonObject>& Params)
+{
+	// Parse required parameters
+	FString WidgetPath;
+	if (!GetStringParam(Params, TEXT("widget_path"), WidgetPath))
+	{
+		return FECACommandResult::Error(TEXT("Missing required parameter: widget_path"));
+	}
+
+	FString WidgetName;
+	if (!GetStringParam(Params, TEXT("widget_name"), WidgetName))
+	{
+		return FECACommandResult::Error(TEXT("Missing required parameter: widget_name"));
+	}
+
+	// Load the Widget Blueprint
+	UWidgetBlueprint* WidgetBlueprint = LoadWidgetBlueprintFromPath(WidgetPath);
+	if (!WidgetBlueprint)
+	{
+		return FECACommandResult::Error(FString::Printf(TEXT("Widget Blueprint not found: %s"), *WidgetPath));
+	}
+
+	UWidgetTree* WidgetTree = WidgetBlueprint->WidgetTree;
+	if (!WidgetTree)
+	{
+		return FECACommandResult::Error(TEXT("Widget Blueprint has no widget tree"));
+	}
+
+	// Find the widget
+	UWidget* Widget = WidgetTree->FindWidget(FName(*WidgetName));
+	if (!Widget)
+	{
+		return FECACommandResult::Error(FString::Printf(TEXT("Widget not found: %s"), *WidgetName));
+	}
+
+	// Get the widget's slot
+	UPanelSlot* Slot = Widget->Slot;
+	if (!Slot)
+	{
+		return FECACommandResult::Error(FString::Printf(TEXT("Widget '%s' has no slot (it may be the root widget or detached)"), *WidgetName));
+	}
+
+	// Track what we set
+	TArray<FString> PropertiesSet;
+	TArray<FString> Warnings;
+	FString SlotTypeName = Slot->GetClass()->GetName();
+
+	// Use UE reflection to set slot properties uniformly across all slot types.
+	// This replaces per-type cast chains for UOverlaySlot, UVerticalBoxSlot,
+	// UHorizontalBoxSlot, USizeBoxSlot, UBorderSlot, UStackBoxSlot, etc.
+	// All these slots expose Padding, HorizontalAlignment, and VerticalAlignment
+	// as UPROPERTYs with the same names and types.
+	//
+	// Ideally UPanelSlot itself would expose virtual setters for these common
+	// properties, but until that engine change happens, reflection works reliably.
+
+	// Parse and apply padding
+	const TSharedPtr<FJsonObject>* PaddingObj = nullptr;
+	if (Params->TryGetObjectField(TEXT("padding"), PaddingObj) && PaddingObj)
+	{
+		FMargin Padding = ParsePadding(*PaddingObj);
+		FStructProperty* Prop = CastField<FStructProperty>(Slot->GetClass()->FindPropertyByName(TEXT("Padding")));
+		if (Prop && Prop->Struct == TBaseStructure<FMargin>::Get())
+		{
+			*Prop->ContainerPtrToValuePtr<FMargin>(Slot) = Padding;
+			Slot->SynchronizeProperties();
+			PropertiesSet.Add(TEXT("padding"));
+		}
+		else
+		{
+			if (Prop)
+			{
+				UE_LOG(LogTemp, Warning, TEXT("ECABridge: Slot '%s' has a Padding property but it is not FMargin (struct: %s). Skipping."), *SlotTypeName, *Prop->Struct->GetName());
+			}
+			Warnings.Add(FString::Printf(TEXT("padding is not supported on slot type: %s"), *SlotTypeName));
+		}
+	}
+
+	// Parse and apply horizontal alignment
+	FString HAlignStr;
+	if (GetStringParam(Params, TEXT("horizontal_alignment"), HAlignStr, false))
+	{
+		EHorizontalAlignment HAlign = ParseHorizontalAlignment(HAlignStr);
+		FProperty* Prop = Slot->GetClass()->FindPropertyByName(TEXT("HorizontalAlignment"));
+		if (Prop && Prop->GetSize() == sizeof(TEnumAsByte<EHorizontalAlignment>))
+		{
+			*Prop->ContainerPtrToValuePtr<TEnumAsByte<EHorizontalAlignment>>(Slot) = HAlign;
+			Slot->SynchronizeProperties();
+			PropertiesSet.Add(TEXT("horizontal_alignment"));
+		}
+		else
+		{
+			if (Prop)
+			{
+				UE_LOG(LogTemp, Warning, TEXT("ECABridge: Slot '%s' has HorizontalAlignment property but unexpected size %d (expected %d). Skipping."), *SlotTypeName, Prop->GetSize(), (int32)sizeof(TEnumAsByte<EHorizontalAlignment>));
+			}
+			Warnings.Add(FString::Printf(TEXT("horizontal_alignment is not supported on slot type: %s"), *SlotTypeName));
+		}
+	}
+
+	// Parse and apply vertical alignment
+	FString VAlignStr;
+	if (GetStringParam(Params, TEXT("vertical_alignment"), VAlignStr, false))
+	{
+		EVerticalAlignment VAlign = ParseVerticalAlignment(VAlignStr);
+		FProperty* Prop = Slot->GetClass()->FindPropertyByName(TEXT("VerticalAlignment"));
+		if (Prop && Prop->GetSize() == sizeof(TEnumAsByte<EVerticalAlignment>))
+		{
+			*Prop->ContainerPtrToValuePtr<TEnumAsByte<EVerticalAlignment>>(Slot) = VAlign;
+			Slot->SynchronizeProperties();
+			PropertiesSet.Add(TEXT("vertical_alignment"));
+		}
+		else
+		{
+			if (Prop)
+			{
+				UE_LOG(LogTemp, Warning, TEXT("ECABridge: Slot '%s' has VerticalAlignment property but unexpected size %d (expected %d). Skipping."), *SlotTypeName, Prop->GetSize(), (int32)sizeof(TEnumAsByte<EVerticalAlignment>));
+			}
+			Warnings.Add(FString::Printf(TEXT("vertical_alignment is not supported on slot type: %s"), *SlotTypeName));
+		}
+	}
+
+	// Parse and apply size rule (only VBox/HBox/StackBox slots have a Size property)
+	FString SizeRuleStr;
+	if (GetStringParam(Params, TEXT("size_rule"), SizeRuleStr, false))
+	{
+		double SizeValue = 1.0;
+		GetFloatParam(Params, TEXT("size_value"), SizeValue, false);
+
+		bool bValidSizeRule = true;
+		FSlateChildSize SlateSize;
+		if (SizeRuleStr.Equals(TEXT("Auto"), ESearchCase::IgnoreCase))
+		{
+			SlateSize = FSlateChildSize(ESlateSizeRule::Automatic);
+		}
+		else if (SizeRuleStr.Equals(TEXT("Fill"), ESearchCase::IgnoreCase))
+		{
+			SlateSize = FSlateChildSize(ESlateSizeRule::Fill);
+			SlateSize.Value = SizeValue;
+		}
+		else
+		{
+			Warnings.Add(FString::Printf(TEXT("Unknown size_rule: '%s'. Use 'Auto' or 'Fill'."), *SizeRuleStr));
+			bValidSizeRule = false;
+		}
+
+		if (bValidSizeRule)
+		{
+			FStructProperty* Prop = CastField<FStructProperty>(Slot->GetClass()->FindPropertyByName(TEXT("Size")));
+			if (Prop && Prop->Struct == FSlateChildSize::StaticStruct())
+			{
+				*Prop->ContainerPtrToValuePtr<FSlateChildSize>(Slot) = SlateSize;
+				Slot->SynchronizeProperties();
+				PropertiesSet.Add(TEXT("size_rule"));
+			}
+			else
+			{
+				if (Prop)
+				{
+					UE_LOG(LogTemp, Warning, TEXT("ECABridge: Slot '%s' has a Size property but it is not FSlateChildSize (struct: %s). Skipping."), *SlotTypeName, *Prop->Struct->GetName());
+				}
+				Warnings.Add(FString::Printf(TEXT("size_rule is not supported on slot type: %s"), *SlotTypeName));
+			}
+		}
+	}
+
+	if (PropertiesSet.Num() == 0 && Warnings.Num() == 0)
+	{
+		return FECACommandResult::Error(TEXT("No slot properties were specified. Provide at least one of: padding, horizontal_alignment, vertical_alignment, size_rule"));
+	}
+
+	// Mark blueprint as modified (not structurally - no tree changes, just property values)
+	FBlueprintEditorUtils::MarkBlueprintAsModified(WidgetBlueprint);
+
+	// Build result
+	TSharedPtr<FJsonObject> Result = MakeResult();
+	Result->SetStringField(TEXT("widget"), WidgetName);
+	Result->SetStringField(TEXT("slot_type"), SlotTypeName);
+
+	TArray<TSharedPtr<FJsonValue>> SetArray;
+	for (const FString& Prop : PropertiesSet)
+	{
+		SetArray.Add(MakeShared<FJsonValueString>(Prop));
+	}
+	Result->SetArrayField(TEXT("properties_set"), SetArray);
+
+	if (Warnings.Num() > 0)
+	{
+		TArray<TSharedPtr<FJsonValue>> WarnArray;
+		for (const FString& Warn : Warnings)
+		{
+			WarnArray.Add(MakeShared<FJsonValueString>(Warn));
+		}
+		Result->SetArrayField(TEXT("warnings"), WarnArray);
+	}
+
+	return FECACommandResult::Success(Result);
+}
+
+//------------------------------------------------------------------------------
+// SetWidgetProperty
+//------------------------------------------------------------------------------
+
+FECACommandResult FECACommand_SetWidgetProperty::Execute(const TSharedPtr<FJsonObject>& Params)
+{
+	// Parse required parameters
+	FString WidgetPath;
+	if (!GetStringParam(Params, TEXT("widget_path"), WidgetPath))
+	{
+		return FECACommandResult::Error(TEXT("Missing required parameter: widget_path"));
+	}
+
+	FString WidgetName;
+	if (!GetStringParam(Params, TEXT("widget_name"), WidgetName))
+	{
+		return FECACommandResult::Error(TEXT("Missing required parameter: widget_name"));
+	}
+
+	const TSharedPtr<FJsonObject>* PropertiesObjPtr = nullptr;
+	if (!Params->TryGetObjectField(TEXT("properties"), PropertiesObjPtr) || !PropertiesObjPtr || !PropertiesObjPtr->IsValid())
+	{
+		return FECACommandResult::Error(TEXT("Missing required parameter: properties (must be a JSON object)"));
+	}
+	const TSharedPtr<FJsonObject>& Properties = *PropertiesObjPtr;
+
+	// Load the Widget Blueprint
+	UWidgetBlueprint* WidgetBlueprint = LoadWidgetBlueprintFromPath(WidgetPath);
+	if (!WidgetBlueprint)
+	{
+		return FECACommandResult::Error(FString::Printf(TEXT("Widget Blueprint not found: %s"), *WidgetPath));
+	}
+
+	UWidgetTree* WidgetTree = WidgetBlueprint->WidgetTree;
+	if (!WidgetTree)
+	{
+		return FECACommandResult::Error(TEXT("Widget Blueprint has no widget tree"));
+	}
+
+	// Find the widget
+	UWidget* Widget = WidgetTree->FindWidget(FName(*WidgetName));
+	if (!Widget)
+	{
+		return FECACommandResult::Error(FString::Printf(TEXT("Widget not found: %s"), *WidgetName));
+	}
+
+	TArray<FString> PropertiesSet;
+	TArray<FString> PropertiesSkipped;
+
+	// Process each property
+	for (const auto& Pair : Properties->Values)
+	{
+		const FString PropName(*Pair.Key);
+		const TSharedPtr<FJsonValue>& PropValue = Pair.Value;
+
+		bool bHandled = false;
+
+		// --- Universal UWidget properties ---
+
+		if (PropName.Equals(TEXT("Visibility"), ESearchCase::IgnoreCase))
+		{
+			FString VisStr;
+			if (PropValue->TryGetString(VisStr))
+			{
+				Widget->SetVisibility(ParseVisibility(VisStr));
+				bHandled = true;
+			}
+		}
+		else if (PropName.Equals(TEXT("RenderOpacity"), ESearchCase::IgnoreCase))
+		{
+			double Opacity = 1.0;
+			if (PropValue->TryGetNumber(Opacity))
+			{
+				Widget->SetRenderOpacity(Opacity);
+				bHandled = true;
+			}
+		}
+		else if (PropName.Equals(TEXT("IsEnabled"), ESearchCase::IgnoreCase))
+		{
+			bool bEnabled = true;
+			if (PropValue->TryGetBool(bEnabled))
+			{
+				Widget->SetIsEnabled(bEnabled);
+				bHandled = true;
+			}
+		}
+		else if (PropName.Equals(TEXT("ToolTipText"), ESearchCase::IgnoreCase))
+		{
+			FString TooltipStr;
+			if (PropValue->TryGetString(TooltipStr))
+			{
+				Widget->SetToolTipText(FText::FromString(TooltipStr));
+				bHandled = true;
+			}
+		}
+		else if (PropName.Equals(TEXT("bIsVariable"), ESearchCase::IgnoreCase))
+		{
+			bool bIsVariable = false;
+			if (PropValue->TryGetBool(bIsVariable))
+			{
+				Widget->bIsVariable = bIsVariable;
+				bHandled = true;
+			}
+		}
+
+		// --- USizeBox properties ---
+		else if (USizeBox* SizeBox = Cast<USizeBox>(Widget))
+		{
+			if (PropName.Equals(TEXT("WidthOverride"), ESearchCase::IgnoreCase))
+			{
+				double Val = 0;
+				if (PropValue->TryGetNumber(Val))
+				{
+					SizeBox->SetWidthOverride(Val);
+					bHandled = true;
+				}
+			}
+			else if (PropName.Equals(TEXT("HeightOverride"), ESearchCase::IgnoreCase))
+			{
+				double Val = 0;
+				if (PropValue->TryGetNumber(Val))
+				{
+					SizeBox->SetHeightOverride(Val);
+					bHandled = true;
+				}
+			}
+			else if (PropName.Equals(TEXT("MinDesiredWidth"), ESearchCase::IgnoreCase))
+			{
+				double Val = 0;
+				if (PropValue->TryGetNumber(Val))
+				{
+					SizeBox->SetMinDesiredWidth(Val);
+					bHandled = true;
+				}
+			}
+			else if (PropName.Equals(TEXT("MinDesiredHeight"), ESearchCase::IgnoreCase))
+			{
+				double Val = 0;
+				if (PropValue->TryGetNumber(Val))
+				{
+					SizeBox->SetMinDesiredHeight(Val);
+					bHandled = true;
+				}
+			}
+			else if (PropName.Equals(TEXT("MaxDesiredWidth"), ESearchCase::IgnoreCase))
+			{
+				double Val = 0;
+				if (PropValue->TryGetNumber(Val))
+				{
+					SizeBox->SetMaxDesiredWidth(Val);
+					bHandled = true;
+				}
+			}
+			else if (PropName.Equals(TEXT("MaxDesiredHeight"), ESearchCase::IgnoreCase))
+			{
+				double Val = 0;
+				if (PropValue->TryGetNumber(Val))
+				{
+					SizeBox->SetMaxDesiredHeight(Val);
+					bHandled = true;
+				}
+			}
+		}
+
+		// --- UTextBlock properties ---
+		else if (UTextBlock* TextBlock = Cast<UTextBlock>(Widget))
+		{
+			if (PropName.Equals(TEXT("Text"), ESearchCase::IgnoreCase))
+			{
+				FString TextStr;
+				if (PropValue->TryGetString(TextStr))
+				{
+					TextBlock->SetText(FText::FromString(TextStr));
+					bHandled = true;
+				}
+			}
+			else if (PropName.Equals(TEXT("ColorAndOpacity"), ESearchCase::IgnoreCase))
+			{
+				const TSharedPtr<FJsonObject>* ColorObj = nullptr;
+				if (PropValue->TryGetObject(ColorObj) && ColorObj)
+				{
+					double R = 1, G = 1, B = 1, A = 1;
+					(*ColorObj)->TryGetNumberField(TEXT("r"), R);
+					(*ColorObj)->TryGetNumberField(TEXT("g"), G);
+					(*ColorObj)->TryGetNumberField(TEXT("b"), B);
+					(*ColorObj)->TryGetNumberField(TEXT("a"), A);
+					TextBlock->SetColorAndOpacity(FSlateColor(FLinearColor(R, G, B, A)));
+					bHandled = true;
+				}
+			}
+		}
+
+		// --- UImage properties ---
+		else if (UImage* Image = Cast<UImage>(Widget))
+		{
+			if (PropName.Equals(TEXT("ColorAndOpacity"), ESearchCase::IgnoreCase))
+			{
+				const TSharedPtr<FJsonObject>* ColorObj = nullptr;
+				if (PropValue->TryGetObject(ColorObj) && ColorObj)
+				{
+					double R = 1, G = 1, B = 1, A = 1;
+					(*ColorObj)->TryGetNumberField(TEXT("r"), R);
+					(*ColorObj)->TryGetNumberField(TEXT("g"), G);
+					(*ColorObj)->TryGetNumberField(TEXT("b"), B);
+					(*ColorObj)->TryGetNumberField(TEXT("a"), A);
+					Image->SetColorAndOpacity(FLinearColor(R, G, B, A));
+					bHandled = true;
+				}
+			}
+			else if (PropName.Equals(TEXT("Opacity"), ESearchCase::IgnoreCase))
+			{
+				double Opacity = 1.0;
+				if (PropValue->TryGetNumber(Opacity))
+				{
+					Image->SetOpacity(Opacity);
+					bHandled = true;
+				}
+			}
+		}
+
+		if (bHandled)
+		{
+			PropertiesSet.Add(PropName);
+		}
+		else
+		{
+			FString WidgetClass = Widget->GetClass()->GetName();
+			PropertiesSkipped.Add(FString::Printf(TEXT("%s (not supported on %s)"), *PropName, *WidgetClass));
+		}
+	}
+
+	if (PropertiesSet.Num() == 0)
+	{
+		FString SkippedList;
+		for (const FString& S : PropertiesSkipped)
+		{
+			if (!SkippedList.IsEmpty()) SkippedList += TEXT(", ");
+			SkippedList += S;
+		}
+		return FECACommandResult::Error(FString::Printf(TEXT("No properties were applied. Skipped: %s"), *SkippedList));
+	}
+
+	// Mark blueprint as modified (not structurally - no tree changes, just property values)
+	FBlueprintEditorUtils::MarkBlueprintAsModified(WidgetBlueprint);
+
+	// Build result
+	TSharedPtr<FJsonObject> Result = MakeResult();
+	Result->SetStringField(TEXT("widget"), WidgetName);
+	Result->SetStringField(TEXT("widget_class"), Widget->GetClass()->GetName());
+
+	TArray<TSharedPtr<FJsonValue>> SetArray;
+	for (const FString& Prop : PropertiesSet)
+	{
+		SetArray.Add(MakeShared<FJsonValueString>(Prop));
+	}
+	Result->SetArrayField(TEXT("properties_set"), SetArray);
+
+	if (PropertiesSkipped.Num() > 0)
+	{
+		TArray<TSharedPtr<FJsonValue>> SkipArray;
+		for (const FString& Skip : PropertiesSkipped)
+		{
+			SkipArray.Add(MakeShared<FJsonValueString>(Skip));
+		}
+		Result->SetArrayField(TEXT("properties_skipped"), SkipArray);
+	}
+
+	return FECACommandResult::Success(Result);
+}
+
+//------------------------------------------------------------------------------
+// AddChildWidget
+//------------------------------------------------------------------------------
+
+FECACommandResult FECACommand_AddChildWidget::Execute(const TSharedPtr<FJsonObject>& Params)
+{
+	// Parse required parameters
+	FString WidgetPath;
+	if (!GetStringParam(Params, TEXT("widget_path"), WidgetPath))
+	{
+		return FECACommandResult::Error(TEXT("Missing required parameter: widget_path"));
+	}
+
+	FString ParentName;
+	if (!GetStringParam(Params, TEXT("parent_name"), ParentName))
+	{
+		return FECACommandResult::Error(TEXT("Missing required parameter: parent_name"));
+	}
+
+	FString WidgetClassName;
+	if (!GetStringParam(Params, TEXT("widget_class"), WidgetClassName))
+	{
+		return FECACommandResult::Error(TEXT("Missing required parameter: widget_class"));
+	}
+
+	FString WidgetName;
+	if (!GetStringParam(Params, TEXT("widget_name"), WidgetName))
+	{
+		return FECACommandResult::Error(TEXT("Missing required parameter: widget_name"));
+	}
+
+	int32 InsertIndex = -1;
+	GetIntParam(Params, TEXT("insert_index"), InsertIndex, false);
+
+	// Load the Widget Blueprint
+	UWidgetBlueprint* WidgetBlueprint = LoadWidgetBlueprintFromPath(WidgetPath);
+	if (!WidgetBlueprint)
+	{
+		return FECACommandResult::Error(FString::Printf(TEXT("Widget Blueprint not found: %s"), *WidgetPath));
+	}
+
+	UWidgetTree* WidgetTree = WidgetBlueprint->WidgetTree;
+	if (!WidgetTree)
+	{
+		return FECACommandResult::Error(TEXT("Widget Blueprint has no widget tree"));
+	}
+
+	// Find the parent widget
+	UWidget* ParentAsWidget = WidgetTree->FindWidget(FName(*ParentName));
+	if (!ParentAsWidget)
+	{
+		return FECACommandResult::Error(FString::Printf(TEXT("Parent widget not found: %s"), *ParentName));
+	}
+
+	UPanelWidget* Parent = Cast<UPanelWidget>(ParentAsWidget);
+	if (!Parent)
+	{
+		return FECACommandResult::Error(FString::Printf(TEXT("Parent '%s' is not a panel widget and cannot accept children"), *ParentName));
+	}
+
+	if (!Parent->CanAddMoreChildren())
+	{
+		return FECACommandResult::Error(FString::Printf(TEXT("Parent '%s' cannot accept more children (it is a %s which already has a child)"), *ParentName, *Parent->GetClass()->GetName()));
+	}
+
+	// Check for name collision
+	if (WidgetTree->FindWidget(FName(*WidgetName)))
+	{
+		return FECACommandResult::Error(FString::Printf(TEXT("A widget named '%s' already exists in the tree"), *WidgetName));
+	}
+
+	// Find the widget class
+	UClass* WidgetClass = FindWidgetClass(WidgetClassName);
+	if (!WidgetClass)
+	{
+		return FECACommandResult::Error(FString::Printf(TEXT("Widget class not found: %s"), *WidgetClassName));
+	}
+
+	if (!WidgetClass->IsChildOf(UWidget::StaticClass()))
+	{
+		return FECACommandResult::Error(FString::Printf(TEXT("Class '%s' is not a UWidget subclass"), *WidgetClassName));
+	}
+
+	// Construct the widget
+	UWidget* NewWidget = WidgetTree->ConstructWidget<UWidget>(WidgetClass, FName(*WidgetName));
+	if (!NewWidget)
+	{
+		return FECACommandResult::Error(FString::Printf(TEXT("Failed to construct widget of class: %s"), *WidgetClassName));
+	}
+
+	NewWidget->SetFlags(RF_Transactional);
+
+	// Add to parent
+	UPanelSlot* NewSlot = nullptr;
+	if (InsertIndex >= 0)
+	{
+		NewSlot = Parent->InsertChildAt(InsertIndex, NewWidget);
+	}
+	else
+	{
+		NewSlot = Parent->AddChild(NewWidget);
+	}
+
+	if (!NewSlot)
+	{
+		WidgetTree->RemoveWidget(NewWidget);
+		return FECACommandResult::Error(TEXT("Failed to add widget to parent"));
+	}
+
+	int32 ActualIndex = Parent->GetChildIndex(NewWidget);
+
+	// Unless skip_compile is set, ensure GUIDs and recompile
+	bool bSkipCompile = false;
+	GetBoolParam(Params, TEXT("skip_compile"), bSkipCompile, false);
+	if (!bSkipCompile)
+	{
+		EnsureAllWidgetGuids(WidgetBlueprint);
+		FBlueprintEditorUtils::MarkBlueprintAsStructurallyModified(WidgetBlueprint);
+	}
+	else
+	{
+		WidgetBlueprint->MarkPackageDirty();
+	}
+
+	// Build result
+	TSharedPtr<FJsonObject> Result = MakeResult();
+	Result->SetStringField(TEXT("widget_name"), WidgetName);
+	Result->SetStringField(TEXT("widget_class"), WidgetClass->GetName());
+	Result->SetStringField(TEXT("parent"), ParentName);
+	Result->SetNumberField(TEXT("index"), ActualIndex);
+	return FECACommandResult::Success(Result);
+}
+
+
+//------------------------------------------------------------------------------
+// CompileWidgetBlueprint
+//------------------------------------------------------------------------------
+
+FECACommandResult FECACommand_CompileWidgetBlueprint::Execute(const TSharedPtr<FJsonObject>& Params)
+{
+	FString WidgetPath;
+	if (!GetStringParam(Params, TEXT("widget_path"), WidgetPath))
+	{
+		return FECACommandResult::Error(TEXT("Missing required parameter: widget_path"));
+	}
+
+	UWidgetBlueprint* WidgetBlueprint = LoadWidgetBlueprintFromPath(WidgetPath);
+	if (!WidgetBlueprint)
+	{
+		return FECACommandResult::Error(FString::Printf(TEXT("Widget Blueprint not found: %s"), *WidgetPath));
+	}
+
+	EnsureAllWidgetGuids(WidgetBlueprint);
+	FBlueprintEditorUtils::MarkBlueprintAsStructurallyModified(WidgetBlueprint);
+
+	// Save the package to disk so changes aren't lost on editor close/crash
+	bool bSaved = false;
+	UPackage* Package = WidgetBlueprint->GetPackage();
+	if (Package)
+	{
+		FString PackageFileName = FPackageName::LongPackageNameToFilename(
+			Package->GetName(), FPackageName::GetAssetPackageExtension());
+		FSavePackageArgs SaveArgs;
+		SaveArgs.TopLevelFlags = RF_Standalone;
+		bSaved = UPackage::SavePackage(Package, WidgetBlueprint, *PackageFileName, SaveArgs);
+		if (bSaved)
+		{
+			UE_LOG(LogTemp, Log, TEXT("ECABridge: Saved %s"), *PackageFileName);
+		}
+		else
+		{
+			UE_LOG(LogTemp, Warning, TEXT("ECABridge: Failed to save %s"), *PackageFileName);
+		}
+	}
+
+	TSharedPtr<FJsonObject> Result = MakeResult();
+	Result->SetStringField(TEXT("widget_path"), WidgetPath);
+	Result->SetBoolField(TEXT("compiled"), true);
+	Result->SetBoolField(TEXT("saved"), bSaved);
+	return FECACommandResult::Success(Result);
+}
+
+//------------------------------------------------------------------------------
+// RemoveWidget
+//------------------------------------------------------------------------------
+
+FECACommandResult FECACommand_RemoveWidget::Execute(const TSharedPtr<FJsonObject>& Params)
+{
+	// Parse required parameters
+	FString WidgetPath;
+	if (!GetStringParam(Params, TEXT("widget_path"), WidgetPath))
+	{
+		return FECACommandResult::Error(TEXT("Missing required parameter: widget_path"));
+	}
+
+	FString WidgetName;
+	if (!GetStringParam(Params, TEXT("widget_name"), WidgetName))
+	{
+		return FECACommandResult::Error(TEXT("Missing required parameter: widget_name"));
+	}
+
+	// Load the Widget Blueprint
+	UWidgetBlueprint* WidgetBlueprint = LoadWidgetBlueprintFromPath(WidgetPath);
+	if (!WidgetBlueprint)
+	{
+		return FECACommandResult::Error(FString::Printf(TEXT("Widget Blueprint not found: %s"), *WidgetPath));
+	}
+
+	UWidgetTree* WidgetTree = WidgetBlueprint->WidgetTree;
+	if (!WidgetTree)
+	{
+		return FECACommandResult::Error(TEXT("Widget Blueprint has no widget tree"));
+	}
+
+	// Find the widget
+	UWidget* Widget = WidgetTree->FindWidget(FName(*WidgetName));
+	if (!Widget)
+	{
+		return FECACommandResult::Error(FString::Printf(TEXT("Widget not found: %s"), *WidgetName));
+	}
+
+	// Don't allow removing the root widget
+	if (Widget == WidgetTree->RootWidget)
+	{
+		return FECACommandResult::Error(TEXT("Cannot remove the root widget. Use wrap_widget to replace the root, or clear the entire tree."));
+	}
+
+	// Count descendants for the response
+	int32 DescendantCount = 0;
+	FString WidgetClass = Widget->GetClass()->GetName();
+	UPanelWidget* Parent = Cast<UPanelWidget>(Widget->GetParent());
+	FString ParentName = Parent ? Parent->GetName() : TEXT("[unknown]");
+
+	if (UPanelWidget* WidgetAsPanel = Cast<UPanelWidget>(Widget))
+	{
+		TArray<UWidget*> Descendants;
+		WidgetTree->GetChildWidgets(Widget, Descendants);
+		DescendantCount = Descendants.Num();
+	}
+
+	// RemoveWidget handles detaching from parent and recursive cleanup
+	WidgetTree->RemoveWidget(Widget);
+
+	// Ensure all remaining widgets have GUIDs, then mark blueprint as modified
+	EnsureAllWidgetGuids(WidgetBlueprint);
+	FBlueprintEditorUtils::MarkBlueprintAsStructurallyModified(WidgetBlueprint);
+
+	// Build result
+	TSharedPtr<FJsonObject> Result = MakeResult();
+	Result->SetStringField(TEXT("removed_widget"), WidgetName);
+	Result->SetStringField(TEXT("widget_class"), WidgetClass);
+	Result->SetStringField(TEXT("parent"), ParentName);
+	Result->SetNumberField(TEXT("descendants_removed"), DescendantCount);
+	return FECACommandResult::Success(Result);
+}
diff --git a/Engine/Plugins/Experimental/ECABridge/Source/ECABridge/Private/ECABridge.cpp b/Engine/Plugins/Experimental/ECABridge/Source/ECABridge/Private/ECABridge.cpp
new file mode 100644
--- /dev/null
+++ b/Engine/Plugins/Experimental/ECABridge/Source/ECABridge/Private/ECABridge.cpp
@@ -0,0 +1,128 @@
+// Copyright Epic Games, Inc. All Rights Reserved.
+
+#include "ECABridge.h"
+#include "ECAMCPServer.h"
+#include "ECABridgeSettings.h"
+#include "Commands/ECACommand.h"
+#include "Dom/JsonObject.h"
+#include "Serialization/JsonReader.h"
+#include "Serialization/JsonSerializer.h"
+
+void UECABridge::Initialize()
+{
+	const UECABridgeSettings* Settings = UECABridgeSettings::Get();
+	
+	if (!Settings->bAutoStart)
+	{
+		UE_LOG(LogTemp, Log, TEXT("[ECABridge] Auto-start disabled in settings"));
+		return;
+	}
+	
+	// Start MCP HTTP server
+	MCPServer = MakeShared<FECAMCPServer>(this);
+	if (MCPServer->Start(Settings->ServerPort))
+	{
+		UE_LOG(LogTemp, Log, TEXT("[ECABridge] ========================================"));
+		UE_LOG(LogTemp, Log, TEXT("[ECABridge] MCP Server running on http://localhost:%d"), Settings->ServerPort);
+		UE_LOG(LogTemp, Log, TEXT("[ECABridge] ========================================"));
+		UE_LOG(LogTemp, Log, TEXT("[ECABridge] Add to Claude Desktop config:"));
+		UE_LOG(LogTemp, Log, TEXT("[ECABridge]   \"unreal-editor\": { \"url\": \"http://localhost:%d/sse\" }"), Settings->ServerPort);
+		UE_LOG(LogTemp, Log, TEXT("[ECABridge] ========================================"));
+	}
+	else
+	{
+		UE_LOG(LogTemp, Error, TEXT("[ECABridge] Failed to start MCP server on port %d"), Settings->ServerPort);
+	}
+}
+
+void UECABridge::Shutdown()
+{
+	if (MCPServer.IsValid())
+	{
+		MCPServer->Stop();
+		MCPServer.Reset();
+	}
+	
+	UE_LOG(LogTemp, Log, TEXT("[ECABridge] Server stopped"));
+}
+
+bool UECABridge::IsRunning() const
+{
+	return MCPServer.IsValid() && MCPServer->IsRunning();
+}
+
+int32 UECABridge::GetPort() const
+{
+	return UECABridgeSettings::Get()->ServerPort;
+}
+
+FString UECABridge::ProcessCommand(const FString& JsonCommand)
+{
+	const UECABridgeSettings* Settings = UECABridgeSettings::Get();
+	
+	if (Settings->bLogCommands)
+	{
+		UE_LOG(LogTemp, Log, TEXT("[ECABridge] Command: %s"), *JsonCommand);
+	}
+	
+	// Parse the command
+	FString CommandName;
+	TSharedPtr<FJsonObject> Params;
+	
+	if (!ParseCommand(JsonCommand, CommandName, Params))
+	{
+		return FECACommandResult::Error(TEXT("Invalid JSON or missing command field")).ToJsonString();
+	}
+	
+	if (Settings->bVerboseLogging)
+	{
+		UE_LOG(LogTemp, Log, TEXT("[ECABridge] Executing: %s"), *CommandName);
+	}
+	
+	// Execute registered command
+	FECACommandResult Result = FECACommandRegistry::Get().ExecuteCommand(CommandName, Params);
+	
+	CommandsProcessed++;
+	
+	if (Settings->bVerboseLogging)
+	{
+		UE_LOG(LogTemp, Log, TEXT("[ECABridge] %s: %s"), 
+			*CommandName, 
+			Result.bSuccess ? TEXT("Success") : *Result.ErrorMessage);
+	}
+	
+	return Result.ToJsonString();
+}
+
+bool UECABridge::ParseCommand(const FString& JsonCommand, FString& OutCommandName, TSharedPtr<FJsonObject>& OutParams)
+{
+	TSharedPtr<FJsonObject> JsonObject;
+	TSharedRef<TJsonReader<>> Reader = TJsonReaderFactory<>::Create(JsonCommand);
+	
+	if (!FJsonSerializer::Deserialize(Reader, JsonObject) || !JsonObject.IsValid())
+	{
+		return false;
+	}
+	
+	if (!JsonObject->TryGetStringField(TEXT("command"), OutCommandName))
+	{
+		return false;
+	}
+	
+	const TSharedPtr<FJsonObject>* ParamsPtr;
+	if (JsonObject->TryGetObjectField(TEXT("params"), ParamsPtr))
+	{
+		OutParams = *ParamsPtr;
+	}
+	else
+	{
+		OutParams = MakeShared<FJsonObject>();
+	}
+	
+	return true;
+}
+
+void UECABridge::LogRegisteredCommands() const
+{
+	FECACommandRegistry::Get().LogCommands();
+}
diff --git a/Engine/Plugins/Experimental/ECABridge/Source/ECABridge/Private/ECABridgeModule.cpp b/Engine/Plugins/Experimental/ECABridge/Source/ECABridge/Private/ECABridgeModule.cpp
new file mode 100644
--- /dev/null
+++ b/Engine/Plugins/Experimental/ECABridge/Source/ECABridge/Private/ECABridgeModule.cpp
@@ -0,0 +1,664 @@
+// Copyright Epic Games, Inc. All Rights Reserved.
+
+#include "ECABridgeModule.h"
+#include "ECABridge.h"
+#include "ECABridgeSettings.h"
+#include "ECAEventQueue.h"
+#include "HAL/IConsoleManager.h"
+#include "Editor.h"
+#include "Selection.h"
+#include "Engine/Selection.h"
+#include "GameFramework/Actor.h"
+#include "LevelEditor.h"
+#include "Modules/ModuleManager.h"
+#include "AssetRegistry/AssetRegistryModule.h"
+#include "Kismet2/KismetEditorUtilities.h"
+#include "Engine/Blueprint.h"
+#include "Subsystems/AssetEditorSubsystem.h"
+#include "Subsystems/ImportSubsystem.h"
+#include "Engine/Brush.h"
+#include "Factories/Factory.h"
+#include "Editor/EditorEngine.h"
+#include "Engine/World.h"
+#include "TimerManager.h"
+#include "UnrealClient.h"  // For ELevelViewportType
+#include "Commands/ECACommand.h"
+#include "Dom/JsonObject.h"
+
+#define LOCTEXT_NAMESPACE "FECABridgeModule"
+
+void FECABridgeModule::StartupModule()
+{
+	UE_LOG(LogTemp, Log, TEXT("[ECABridge] Module starting..."));
+	
+	// Create and initialize the bridge with the transient package as outer
+	Bridge = NewObject<UECABridge>(GetTransientPackage(), NAME_None, RF_Transient);
+	Bridge->Initialize();
+	
+	// Register console commands
+	RegisterConsoleCommands();
+	
+	// Register editor event handlers
+	RegisterEditorEventHandlers();
+	
+	UE_LOG(LogTemp, Log, TEXT("[ECABridge] Module ready"));
+}
+
+void FECABridgeModule::ShutdownModule()
+{
+	// Unregister editor event handlers
+	UnregisterEditorEventHandlers();
+	
+	if (Bridge)
+	{
+		Bridge->Shutdown();
+		Bridge = nullptr;
+	}
+	
+	UE_LOG(LogTemp, Log, TEXT("[ECABridge] Module shutdown"));
+}
+
+void FECABridgeModule::RegisterConsoleCommands()
+{
+	IConsoleManager::Get().RegisterConsoleCommand(
+		TEXT("ECA.Status"),
+		TEXT("Show ECA Bridge server status"),
+		FConsoleCommandDelegate::CreateLambda([]()
+		{
+			if (FECABridgeModule::IsAvailable())
+			{
+				UECABridge* B = FECABridgeModule::Get().GetBridge();
+				if (B)
+				{
+					UE_LOG(LogTemp, Log, TEXT("[ECABridge] Status: %s"), 
+						B->IsRunning() ? TEXT("Running") : TEXT("Stopped"));
+					UE_LOG(LogTemp, Log, TEXT("[ECABridge] URL: http://localhost:%d/sse"), B->GetPort());
+					UE_LOG(LogTemp, Log, TEXT("[ECABridge] Commands processed: %d"), B->GetCommandsProcessed());
+				}
+			}
+		}),
+		ECVF_Default
+	);
+	
+	IConsoleManager::Get().RegisterConsoleCommand(
+		TEXT("ECA.ListCommands"),
+		TEXT("List all available ECA commands"),
+		FConsoleCommandDelegate::CreateLambda([]()
+		{
+			if (FECABridgeModule::IsAvailable())
+			{
+				UECABridge* B = FECABridgeModule::Get().GetBridge();
+				if (B)
+				{
+					B->LogRegisteredCommands();
+				}
+			}
+		}),
+		ECVF_Default
+	);
+}
+
+void FECABridgeModule::RegisterEditorEventHandlers()
+{
+	// Register for actor selection changes (static delegate, always available)
+	SelectionChangedHandle = USelection::SelectionChangedEvent.AddRaw(this, &FECABridgeModule::OnEditorSelectionChanged);
+	UE_LOG(LogTemp, Log, TEXT("[ECABridge] Registered selection change handler"));
+	
+	// Register for asset registry events (covers asset creation, rename, delete)
+	if (FModuleManager::Get().IsModuleLoaded("AssetRegistry"))
+	{
+		IAssetRegistry& AssetRegistry = FModuleManager::GetModuleChecked<FAssetRegistryModule>("AssetRegistry").Get();
+		AssetAddedHandle = AssetRegistry.OnAssetAdded().AddRaw(this, &FECABridgeModule::OnAssetAdded);
+		AssetRemovedHandle = AssetRegistry.OnAssetRemoved().AddRaw(this, &FECABridgeModule::OnAssetRemoved);
+		AssetRenamedHandle = AssetRegistry.OnAssetRenamed().AddRaw(this, &FECABridgeModule::OnAssetRenamed);
+		UE_LOG(LogTemp, Log, TEXT("[ECABridge] Registered asset registry handlers"));
+	}
+	
+	// Register for level/map change events (static delegates, always available)
+	MapChangedHandle = FEditorDelegates::OnMapOpened.AddRaw(this, &FECABridgeModule::OnMapOpened);
+	UE_LOG(LogTemp, Log, TEXT("[ECABridge] Registered map change handler"));
+	
+	// Register for PIE (Play In Editor) events (static delegates, always available)
+	PIEStartedHandle = FEditorDelegates::BeginPIE.AddRaw(this, &FECABridgeModule::OnPIEStarted);
+	PIEEndedHandle = FEditorDelegates::EndPIE.AddRaw(this, &FECABridgeModule::OnPIEEnded);
+	UE_LOG(LogTemp, Log, TEXT("[ECABridge] Registered PIE handlers"));
+	
+	// GEditor-dependent registrations need to be deferred until editor is ready
+	// We'll register these when GEditor becomes available
+	if (GEditor)
+	{
+		RegisterGEditorDependentHandlers();
+	}
+	else
+	{
+		// Defer registration until editor is initialized
+		FEditorDelegates::OnEditorInitialized.AddRaw(this, &FECABridgeModule::OnEditorInitialized);
+		UE_LOG(LogTemp, Log, TEXT("[ECABridge] Deferred GEditor-dependent handlers until editor init"));
+	}
+	
+	UE_LOG(LogTemp, Log, TEXT("[ECABridge] Editor event handlers registration complete"));
+}
+
+void FECABridgeModule::OnEditorInitialized(double Duration)
+{
+	UE_LOG(LogTemp, Log, TEXT("[ECABridge] Editor initialized (took %.2fs), registering deferred handlers"), Duration);
+	RegisterGEditorDependentHandlers();
+}
+
+void FECABridgeModule::RegisterGEditorDependentHandlers()
+{
+	if (!GEditor)
+	{
+		UE_LOG(LogTemp, Warning, TEXT("[ECABridge] GEditor still null, cannot register GEditor-dependent handlers"));
+		return;
+	}
+	
+	// Register for asset import events
+	if (UImportSubsystem* ImportSubsystem = GEditor->GetEditorSubsystem<UImportSubsystem>())
+	{
+		AssetImportHandle = ImportSubsystem->OnAssetPostImport.AddRaw(this, &FECABridgeModule::OnAssetImported);
+		UE_LOG(LogTemp, Log, TEXT("[ECABridge] Registered asset import handler"));
+	}
+	
+	// Register for Blueprint compilation events
+	BlueprintCompiledHandle = GEditor->OnBlueprintCompiled().AddRaw(this, &FECABridgeModule::OnBlueprintCompiled);
+	UE_LOG(LogTemp, Log, TEXT("[ECABridge] Registered Blueprint compilation handler"));
+	
+	// Register for actor spawn/delete in editor
+	ActorSpawnedHandle = GEditor->OnLevelActorAdded().AddRaw(this, &FECABridgeModule::OnActorSpawned);
+	ActorDeletedHandle = GEditor->OnLevelActorDeleted().AddRaw(this, &FECABridgeModule::OnActorDeleted);
+	UE_LOG(LogTemp, Log, TEXT("[ECABridge] Registered actor spawn/delete handlers"));
+	
+	// Register for editor camera movement events
+	CameraMovedHandle = FEditorDelegates::OnEditorCameraMoved.AddRaw(this, &FECABridgeModule::OnEditorCameraMoved);
+	UE_LOG(LogTemp, Log, TEXT("[ECABridge] Registered camera movement handler"));
+	
+	UE_LOG(LogTemp, Log, TEXT("[ECABridge] All GEditor-dependent handlers registered"));
+}
+
+void FECABridgeModule::UnregisterEditorEventHandlers()
+{
+	// Unregister selection change handler
+	if (SelectionChangedHandle.IsValid())
+	{
+		USelection::SelectionChangedEvent.Remove(SelectionChangedHandle);
+		SelectionChangedHandle.Reset();
+	}
+	
+	// Unregister asset import handler
+	if (AssetImportHandle.IsValid() && GEditor)
+	{
+		if (UImportSubsystem* ImportSubsystem = GEditor->GetEditorSubsystem<UImportSubsystem>())
+		{
+			ImportSubsystem->OnAssetPostImport.Remove(AssetImportHandle);
+		}
+		AssetImportHandle.Reset();
+	}
+	
+	// Unregister asset registry handlers
+	if (FModuleManager::Get().IsModuleLoaded("AssetRegistry"))
+	{
+		IAssetRegistry& AssetRegistry = FModuleManager::GetModuleChecked<FAssetRegistryModule>("AssetRegistry").Get();
+		if (AssetAddedHandle.IsValid())
+		{
+			AssetRegistry.OnAssetAdded().Remove(AssetAddedHandle);
+			AssetAddedHandle.Reset();
+		}
+		if (AssetRemovedHandle.IsValid())
+		{
+			AssetRegistry.OnAssetRemoved().Remove(AssetRemovedHandle);
+			AssetRemovedHandle.Reset();
+		}
+		if (AssetRenamedHandle.IsValid())
+		{
+			AssetRegistry.OnAssetRenamed().Remove(AssetRenamedHandle);
+			AssetRenamedHandle.Reset();
+		}
+	}
+	
+	// Unregister Blueprint compilation handler
+	if (BlueprintCompiledHandle.IsValid() && GEditor)
+	{
+		GEditor->OnBlueprintCompiled().Remove(BlueprintCompiledHandle);
+		BlueprintCompiledHandle.Reset();
+	}
+	
+	// Unregister map change handler
+	if (MapChangedHandle.IsValid())
+	{
+		FEditorDelegates::OnMapOpened.Remove(MapChangedHandle);
+		MapChangedHandle.Reset();
+	}
+	
+	// Unregister PIE handlers
+	if (PIEStartedHandle.IsValid())
+	{
+		FEditorDelegates::BeginPIE.Remove(PIEStartedHandle);
+		PIEStartedHandle.Reset();
+	}
+	if (PIEEndedHandle.IsValid())
+	{
+		FEditorDelegates::EndPIE.Remove(PIEEndedHandle);
+		PIEEndedHandle.Reset();
+	}
+	
+	// Unregister actor spawn/delete handlers
+	if (GEditor)
+	{
+		if (ActorSpawnedHandle.IsValid())
+		{
+			GEditor->OnLevelActorAdded().Remove(ActorSpawnedHandle);
+			ActorSpawnedHandle.Reset();
+		}
+		if (ActorDeletedHandle.IsValid())
+		{
+			GEditor->OnLevelActorDeleted().Remove(ActorDeletedHandle);
+			ActorDeletedHandle.Reset();
+		}
+	}
+	
+	// Unregister camera movement handler
+	if (CameraMovedHandle.IsValid())
+	{
+		FEditorDelegates::OnEditorCameraMoved.Remove(CameraMovedHandle);
+		CameraMovedHandle.Reset();
+	}
+	
+	// Clear any pending camera timer
+	if (GEditor && GEditor->GetEditorWorldContext().World())
+	{
+		GEditor->GetEditorWorldContext().World()->GetTimerManager().ClearTimer(CameraMovedTimerHandle);
+	}
+	
+	UE_LOG(LogTemp, Log, TEXT("[ECABridge] All editor event handlers unregistered"));
+}
+
+void FECABridgeModule::OnEditorSelectionChanged(UObject* SelectionObject)
+{
+	UE_LOG(LogTemp, Log, TEXT("[ECABridge] OnEditorSelectionChanged called! Object: %s"), 
+		SelectionObject ? *SelectionObject->GetName() : TEXT("null"));
+	
+	if (!GEditor)
+	{
+		UE_LOG(LogTemp, Warning, TEXT("[ECABridge] GEditor is null"));
+		return;
+	}
+	
+	// Only process actor selection changes (not component selection, etc.)
+	USelection* ActorSelection = GEditor->GetSelectedActors();
+	if (!ActorSelection)
+	{
+		UE_LOG(LogTemp, Warning, TEXT("[ECABridge] GetSelectedActors() returned null"));
+		return;
+	}
+	
+	// Check if this is actor selection (vs component selection, etc.)
+	if (SelectionObject != ActorSelection)
+	{
+		UE_LOG(LogTemp, Verbose, TEXT("[ECABridge] Ignoring non-actor selection change"));
+		return;
+	}
+	
+	// Get selected actors
+	TArray<AActor*> SelectedActors;
+	ActorSelection->GetSelectedObjects<AActor>(SelectedActors);
+	
+	UE_LOG(LogTemp, Log, TEXT("[ECABridge] Actor selection changed: %d actors"), SelectedActors.Num());
+	
+	if (SelectedActors.Num() == 0)
+	{
+		// Selection cleared - queue a deselection event
+		TSharedPtr<FJsonObject> Data = MakeShared<FJsonObject>();
+		Data->SetNumberField(TEXT("count"), 0);
+		Data->SetArrayField(TEXT("actors"), TArray<TSharedPtr<FJsonValue>>());
+		ECAEvents::QueueCustomEvent(TEXT("selection_changed"), TEXT("Selection cleared"), Data, false);
+		return;
+	}
+	
+	// Build selection data
+	TSharedPtr<FJsonObject> Data = MakeShared<FJsonObject>();
+	Data->SetNumberField(TEXT("count"), SelectedActors.Num());
+	
+	TArray<TSharedPtr<FJsonValue>> ActorsArray;
+	for (AActor* Actor : SelectedActors)
+	{
+		if (Actor)
+		{
+			TSharedPtr<FJsonObject> ActorInfo = MakeShared<FJsonObject>();
+			ActorInfo->SetStringField(TEXT("name"), Actor->GetActorNameOrLabel());
+			ActorInfo->SetStringField(TEXT("class"), Actor->GetClass()->GetName());
+			ActorInfo->SetStringField(TEXT("path"), Actor->GetPathName());
+			
+			// Add location
+			FVector Location = Actor->GetActorLocation();
+			ActorInfo->SetNumberField(TEXT("x"), Location.X);
+			ActorInfo->SetNumberField(TEXT("y"), Location.Y);
+			ActorInfo->SetNumberField(TEXT("z"), Location.Z);
+			
+			ActorsArray.Add(MakeShared<FJsonValueObject>(ActorInfo));
+		}
+	}
+	Data->SetArrayField(TEXT("actors"), ActorsArray);
+	
+	// Queue the event - single selection uses the simpler helper, multi-selection uses custom
+	if (SelectedActors.Num() == 1)
+	{
+		AActor* Actor = SelectedActors[0];
+		ECAEvents::QueueSelectionChanged(Actor->GetActorNameOrLabel(), Actor->GetPathName());
+	}
+	else
+	{
+		FString Message = FString::Printf(TEXT("Selected %d actors"), SelectedActors.Num());
+		ECAEvents::QueueCustomEvent(TEXT("selection_changed"), Message, Data, true);
+	}
+}
+
+void FECABridgeModule::OnAssetImported(UFactory* Factory, UObject* CreatedObject)
+{
+	if (!CreatedObject)
+	{
+		return;
+	}
+	
+	FString AssetPath = CreatedObject->GetPathName();
+	FString AssetType = CreatedObject->GetClass()->GetName();
+	FString AssetName = CreatedObject->GetName();
+	
+	UE_LOG(LogTemp, Log, TEXT("[ECABridge] Asset imported: %s (%s)"), *AssetPath, *AssetType);
+	
+	TSharedPtr<FJsonObject> Data = MakeShared<FJsonObject>();
+	Data->SetStringField(TEXT("asset_path"), AssetPath);
+	Data->SetStringField(TEXT("asset_name"), AssetName);
+	Data->SetStringField(TEXT("asset_type"), AssetType);
+	if (Factory)
+	{
+		Data->SetStringField(TEXT("factory"), Factory->GetClass()->GetName());
+	}
+	
+	FString Message = FString::Printf(TEXT("Imported asset: %s (%s)"), *AssetName, *AssetType);
+	ECAEvents::QueueCustomEvent(TEXT("asset_imported"), Message, Data, false);
+}
+
+void FECABridgeModule::OnAssetAdded(const FAssetData& AssetData)
+{
+	// Skip transient or temporary assets
+	if (AssetData.PackagePath.ToString().StartsWith(TEXT("/Temp")) ||
+		AssetData.PackagePath.ToString().StartsWith(TEXT("/Engine")))
+	{
+		return;
+	}
+	
+	FString AssetPath = AssetData.GetObjectPathString();
+	FString AssetType = AssetData.AssetClassPath.GetAssetName().ToString();
+	FString AssetName = AssetData.AssetName.ToString();
+	
+	UE_LOG(LogTemp, Log, TEXT("[ECABridge] Asset added: %s (%s)"), *AssetPath, *AssetType);
+	
+	TSharedPtr<FJsonObject> Data = MakeShared<FJsonObject>();
+	Data->SetStringField(TEXT("asset_path"), AssetPath);
+	Data->SetStringField(TEXT("asset_name"), AssetName);
+	Data->SetStringField(TEXT("asset_type"), AssetType);
+	Data->SetStringField(TEXT("package_path"), AssetData.PackagePath.ToString());
+	
+	FString Message = FString::Printf(TEXT("Asset created: %s (%s)"), *AssetName, *AssetType);
+	ECAEvents::QueueCustomEvent(TEXT("asset_added"), Message, Data, false);
+}
+
+void FECABridgeModule::OnAssetRemoved(const FAssetData& AssetData)
+{
+	FString AssetPath = AssetData.GetObjectPathString();
+	FString AssetName = AssetData.AssetName.ToString();
+	
+	UE_LOG(LogTemp, Log, TEXT("[ECABridge] Asset removed: %s"), *AssetPath);
+	
+	TSharedPtr<FJsonObject> Data = MakeShared<FJsonObject>();
+	Data->SetStringField(TEXT("asset_path"), AssetPath);
+	Data->SetStringField(TEXT("asset_name"), AssetName);
+	
+	FString Message = FString::Printf(TEXT("Asset deleted: %s"), *AssetName);
+	ECAEvents::QueueCustomEvent(TEXT("asset_removed"), Message, Data, false);
+}
+
+void FECABridgeModule::OnAssetRenamed(const FAssetData& AssetData, const FString& OldObjectPath)
+{
+	FString NewPath = AssetData.GetObjectPathString();
+	FString AssetName = AssetData.AssetName.ToString();
+	
+	UE_LOG(LogTemp, Log, TEXT("[ECABridge] Asset renamed: %s -> %s"), *OldObjectPath, *NewPath);
+	
+	TSharedPtr<FJsonObject> Data = MakeShared<FJsonObject>();
+	Data->SetStringField(TEXT("old_path"), OldObjectPath);
+	Data->SetStringField(TEXT("new_path"), NewPath);
+	Data->SetStringField(TEXT("asset_name"), AssetName);
+	
+	FString Message = FString::Printf(TEXT("Asset renamed: %s"), *AssetName);
+	ECAEvents::QueueCustomEvent(TEXT("asset_renamed"), Message, Data, false);
+}
+
+void FECABridgeModule::OnBlueprintCompiled()
+{
+	// This delegate doesn't give us the specific Blueprint, so we log a generic message
+	// For more specific Blueprint compilation tracking, we'd need to hook into FKismetCompilerContext
+	UE_LOG(LogTemp, Log, TEXT("[ECABridge] A Blueprint was compiled"));
+	
+	TSharedPtr<FJsonObject> Data = MakeShared<FJsonObject>();
+	ECAEvents::QueueCustomEvent(TEXT("blueprint_compiled"), TEXT("A Blueprint was compiled"), Data, false);
+}
+
+void FECABridgeModule::OnMapOpened(const FString& Filename, bool bAsTemplate)
+{
+	UE_LOG(LogTemp, Log, TEXT("[ECABridge] Map opened: %s (template: %s)"), *Filename, bAsTemplate ? TEXT("true") : TEXT("false"));
+	
+	TSharedPtr<FJsonObject> Data = MakeShared<FJsonObject>();
+	Data->SetStringField(TEXT("filename"), Filename);
+	Data->SetBoolField(TEXT("as_template"), bAsTemplate);
+	
+	// Extract just the map name from the path
+	FString MapName = FPaths::GetBaseFilename(Filename);
+	
+	FString Message = FString::Printf(TEXT("Opened map: %s"), *MapName);
+	ECAEvents::QueueCustomEvent(TEXT("map_opened"), Message, Data, false);
+}
+
+void FECABridgeModule::OnPIEStarted(bool bIsSimulating)
+{
+	UE_LOG(LogTemp, Log, TEXT("[ECABridge] PIE started (simulating: %s)"), bIsSimulating ? TEXT("true") : TEXT("false"));
+	
+	TSharedPtr<FJsonObject> Data = MakeShared<FJsonObject>();
+	Data->SetBoolField(TEXT("is_simulating"), bIsSimulating);
+	
+	FString Message = bIsSimulating ? TEXT("Started Simulate In Editor") : TEXT("Started Play In Editor");
+	ECAEvents::QueueCustomEvent(TEXT("pie_started"), Message, Data, false);
+}
+
+void FECABridgeModule::OnPIEEnded(bool bIsSimulating)
+{
+	UE_LOG(LogTemp, Log, TEXT("[ECABridge] PIE ended (simulating: %s)"), bIsSimulating ? TEXT("true") : TEXT("false"));
+	
+	TSharedPtr<FJsonObject> Data = MakeShared<FJsonObject>();
+	Data->SetBoolField(TEXT("is_simulating"), bIsSimulating);
+	
+	FString Message = bIsSimulating ? TEXT("Stopped Simulate In Editor") : TEXT("Stopped Play In Editor");
+	ECAEvents::QueueCustomEvent(TEXT("pie_ended"), Message, Data, false);
+}
+
+void FECABridgeModule::OnActorSpawned(AActor* Actor)
+{
+	if (!Actor)
+	{
+		return;
+	}
+	
+	// Skip transient actors and editor-only actors that are constantly created
+	if (Actor->HasAnyFlags(RF_Transient) || Actor->IsA(ABrush::StaticClass()))
+	{
+		return;
+	}
+	
+	FString ActorName = Actor->GetActorNameOrLabel();
+	FString ActorClass = Actor->GetClass()->GetName();
+	FString ActorPath = Actor->GetPathName();
+	
+	UE_LOG(LogTemp, Log, TEXT("[ECABridge] Actor spawned: %s (%s)"), *ActorName, *ActorClass);
+	
+	TSharedPtr<FJsonObject> Data = MakeShared<FJsonObject>();
+	Data->SetStringField(TEXT("actor_name"), ActorName);
+	Data->SetStringField(TEXT("actor_class"), ActorClass);
+	Data->SetStringField(TEXT("actor_path"), ActorPath);
+	
+	FVector Location = Actor->GetActorLocation();
+	Data->SetNumberField(TEXT("x"), Location.X);
+	Data->SetNumberField(TEXT("y"), Location.Y);
+	Data->SetNumberField(TEXT("z"), Location.Z);
+	
+	FString Message = FString::Printf(TEXT("Actor spawned: %s (%s)"), *ActorName, *ActorClass);
+	ECAEvents::QueueCustomEvent(TEXT("actor_spawned"), Message, Data, false);
+}
+
+void FECABridgeModule::OnActorDeleted(AActor* Actor)
+{
+	if (!Actor)
+	{
+		return;
+	}
+	
+	// Skip transient actors
+	if (Actor->HasAnyFlags(RF_Transient))
+	{
+		return;
+	}
+	
+	FString ActorName = Actor->GetActorNameOrLabel();
+	FString ActorClass = Actor->GetClass()->GetName();
+	
+	UE_LOG(LogTemp, Log, TEXT("[ECABridge] Actor deleted: %s (%s)"), *ActorName, *ActorClass);
+	
+	TSharedPtr<FJsonObject> Data = MakeShared<FJsonObject>();
+	Data->SetStringField(TEXT("actor_name"), ActorName);
+	Data->SetStringField(TEXT("actor_class"), ActorClass);
+	
+	FString Message = FString::Printf(TEXT("Actor deleted: %s (%s)"), *ActorName, *ActorClass);
+	ECAEvents::QueueCustomEvent(TEXT("actor_deleted"), Message, Data, false);
+}
+
+// Helper functions for JSON conversion in camera movement tracking
+namespace ECABridgeHelpers
+{
+	TSharedPtr<FJsonObject> VectorToJson(const FVector& Vec)
+	{
+		TSharedPtr<FJsonObject> Json = MakeShared<FJsonObject>();
+		Json->SetNumberField(TEXT("x"), Vec.X);
+		Json->SetNumberField(TEXT("y"), Vec.Y);
+		Json->SetNumberField(TEXT("z"), Vec.Z);
+		return Json;
+	}
+	
+	TSharedPtr<FJsonObject> RotatorToJson(const FRotator& Rot)
+	{
+		TSharedPtr<FJsonObject> Json = MakeShared<FJsonObject>();
+		Json->SetNumberField(TEXT("pitch"), Rot.Pitch);
+		Json->SetNumberField(TEXT("yaw"), Rot.Yaw);
+		Json->SetNumberField(TEXT("roll"), Rot.Roll);
+		return Json;
+	}
+}
+
+void FECABridgeModule::OnEditorCameraMoved(const FVector& Location, const FRotator& Rotation, ELevelViewportType ViewportType, int32 ViewIndex)
+{
+	// Only track perspective camera movements for now
+	if (ViewportType != LVT_Perspective)
+	{
+		return;
+	}
+	
+	// Store the starting position if this is the beginning of a movement
+	if (!bCameraIsMoving)
+	{
+		bCameraIsMoving = true;
+		LastCameraLocation = Location;
+		LastCameraRotation = Rotation;
+	}
+	
+	// Reset/extend the debounce timer
+	// Every time the camera moves, we restart the timer
+	if (GEditor && GEditor->GetEditorWorldContext().World())
+	{
+		UWorld* EditorWorld = GEditor->GetEditorWorldContext().World();
+		
+		// Clear any existing timer
+		EditorWorld->GetTimerManager().ClearTimer(CameraMovedTimerHandle);
+		
+		// Set a new timer - when this fires without being cleared, the camera has stopped
+		EditorWorld->GetTimerManager().SetTimer(
+			CameraMovedTimerHandle,
+			FTimerDelegate::CreateRaw(this, &FECABridgeModule::OnCameraStoppedMoving),
+			CameraStoppedDebounceTime,
+			false  // Don't loop
+		);
+	}
+}
+
+void FECABridgeModule::OnCameraStoppedMoving()
+{
+	if (!bCameraIsMoving)
+	{
+		return;
+	}
+	
+	bCameraIsMoving = false;
+	
+	UE_LOG(LogTemp, Log, TEXT("[ECABridge] Camera stopped moving - triggering describe_view"));
+	
+	// Execute the describe_view command to get scene information
+	TSharedPtr<FJsonObject> Params = MakeShared<FJsonObject>();
+	// Use default parameters for describe_view
+	
+	// Find and execute the describe_view command
+	TSharedPtr<IECACommand> DescribeViewCmd = FECACommandRegistry::Get().GetCommand(TEXT("describe_view"));
+	if (DescribeViewCmd.IsValid())
+	{
+		FECACommandResult Result = DescribeViewCmd->Execute(Params);
+		
+		if (Result.bSuccess && Result.ResultData.IsValid())
+		{
+			// Create an event with the describe_view data
+			TSharedPtr<FJsonObject> EventData = MakeShared<FJsonObject>();
+			
+			// Include movement information
+			EventData->SetObjectField(TEXT("start_location"), ECABridgeHelpers::VectorToJson(LastCameraLocation));
+			EventData->SetObjectField(TEXT("start_rotation"), ECABridgeHelpers::RotatorToJson(LastCameraRotation));
+			
+			// Include the describe_view result
+			EventData->SetObjectField(TEXT("view_description"), Result.ResultData);
+			
+			// Queue the event - this will trigger AI response so it can describe what's now visible
+			ECAEvents::QueueCustomEvent(
+				TEXT("camera_moved"),
+				TEXT("User moved the editor viewport camera"),
+				EventData,
+				true  // Trigger AI response to describe the new view
+			);
+		}
+		else
+		{
+			// If describe_view failed, still queue a basic event
+			TSharedPtr<FJsonObject> EventData = MakeShared<FJsonObject>();
+			EventData->SetStringField(TEXT("error"), Result.ErrorMessage);
+			
+			ECAEvents::QueueCustomEvent(
+				TEXT("camera_moved"),
+				TEXT("User moved the editor viewport camera (view description unavailable)"),
+				EventData,
+				false  // Don't trigger AI response on failure
+			);
+		}
+	}
+	else
+	{
+		UE_LOG(LogTemp, Warning, TEXT("[ECABridge] describe_view command not found"));
+	}
+}
+
+#undef LOCTEXT_NAMESPACE
+
+IMPLEMENT_MODULE(FECABridgeModule, ECABridge)
diff --git a/Engine/Plugins/Experimental/ECABridge/Source/ECABridge/Private/ECABridgeSettings.cpp b/Engine/Plugins/Experimental/ECABridge/Source/ECABridge/Private/ECABridgeSettings.cpp
new file mode 100644
--- /dev/null
+++ b/Engine/Plugins/Experimental/ECABridge/Source/ECABridge/Private/ECABridgeSettings.cpp
@@ -0,0 +1,8 @@
+// Copyright Epic Games, Inc. All Rights Reserved.
+
+#include "ECABridgeSettings.h"
+
+UECABridgeSettings::UECABridgeSettings()
+{
+	// Default settings are set in the header
+}
diff --git a/Engine/Plugins/Experimental/ECABridge/Source/ECABridge/Private/ECAMCPServer.cpp b/Engine/Plugins/Experimental/ECABridge/Source/ECABridge/Private/ECAMCPServer.cpp
new file mode 100644
--- /dev/null
+++ b/Engine/Plugins/Experimental/ECABridge/Source/ECABridge/Private/ECAMCPServer.cpp
@@ -0,0 +1,541 @@
+// Copyright Epic Games, Inc. All Rights Reserved.
+
+#include "ECAMCPServer.h"
+#include "ECABridge.h"
+#include "Commands/ECACommand.h"
+#include "HttpServerModule.h"
+#include "IHttpRouter.h"
+#include "HttpPath.h"
+#include "Dom/JsonObject.h"
+#include "Dom/JsonValue.h"
+#include "Serialization/JsonReader.h"
+#include "Serialization/JsonWriter.h"
+#include "Serialization/JsonSerializer.h"
+
+FECAMCPServer::FECAMCPServer(UECABridge* InBridge)
+	: Bridge(InBridge)
+{
+}
+
+FECAMCPServer::~FECAMCPServer()
+{
+	Stop();
+}
+
+bool FECAMCPServer::Start(int32 Port)
+{
+	if (bIsRunning)
+	{
+		UE_LOG(LogTemp, Warning, TEXT("[ECABridge] MCP Server already running"));
+		return true;
+	}
+
+	ServerPort = Port;
+
+	// Get or start HTTP server module
+	FHttpServerModule& HttpServerModule = FHttpServerModule::Get();
+	HttpRouter = HttpServerModule.GetHttpRouter(ServerPort);
+	
+	if (!HttpRouter.IsValid())
+	{
+		UE_LOG(LogTemp, Error, TEXT("[ECABridge] Failed to get HTTP router for port %d"), ServerPort);
+		return false;
+	}
+
+	// Register routes
+	
+	// Main MCP endpoint (POST for messages, OPTIONS for CORS)
+	MCPRouteHandle = HttpRouter->BindRoute(
+		FHttpPath(TEXT("/mcp")),
+		EHttpServerRequestVerbs::VERB_POST | EHttpServerRequestVerbs::VERB_OPTIONS,
+		FHttpRequestHandler::CreateLambda([this](const FHttpServerRequest& Request, const FHttpResultCallback& OnComplete)
+		{
+			if (Request.Verb == EHttpServerRequestVerbs::VERB_OPTIONS)
+			{
+				return HandleCORSPreflight(Request, OnComplete);
+			}
+			return HandleMCPPost(Request, OnComplete);
+		})
+	);
+
+	// Health check endpoint
+	HealthRouteHandle = HttpRouter->BindRoute(
+		FHttpPath(TEXT("/health")),
+		EHttpServerRequestVerbs::VERB_GET | EHttpServerRequestVerbs::VERB_OPTIONS,
+		FHttpRequestHandler::CreateLambda([this](const FHttpServerRequest& Request, const FHttpResultCallback& OnComplete)
+		{
+			if (Request.Verb == EHttpServerRequestVerbs::VERB_OPTIONS)
+			{
+				return HandleCORSPreflight(Request, OnComplete);
+			}
+			return HandleHealth(Request, OnComplete);
+		})
+	);
+
+	// Start listening
+	HttpServerModule.StartAllListeners();
+
+	bIsRunning = true;
+	UE_LOG(LogTemp, Log, TEXT("[ECABridge] MCP Streamable HTTP Server started on port %d"), ServerPort);
+	UE_LOG(LogTemp, Log, TEXT("[ECABridge] MCP endpoint: http://localhost:%d/mcp"), ServerPort);
+
+	return true;
+}
+
+void FECAMCPServer::Stop()
+{
+	if (!bIsRunning)
+	{
+		return;
+	}
+
+	// Unbind routes
+	if (HttpRouter.IsValid())
+	{
+		HttpRouter->UnbindRoute(MCPRouteHandle);
+		HttpRouter->UnbindRoute(HealthRouteHandle);
+	}
+
+	bIsRunning = false;
+	UE_LOG(LogTemp, Log, TEXT("[ECABridge] MCP Server stopped"));
+}
+
+bool FECAMCPServer::HandleMCPPost(const FHttpServerRequest& Request, const FHttpResultCallback& OnComplete)
+{
+	// Parse request body
+	FString RequestBody;
+	if (Request.Body.Num() > 0)
+	{
+		FUTF8ToTCHAR Converter(reinterpret_cast<const ANSICHAR*>(Request.Body.GetData()), Request.Body.Num());
+		RequestBody = FString(Converter.Length(), Converter.Get());
+	}
+	
+	UE_LOG(LogTemp, Log, TEXT("[ECABridge] MCP Request: %s"), *RequestBody.Left(500));
+	
+	if (RequestBody.IsEmpty())
+	{
+		TSharedPtr<FJsonObject> ErrorResponse = CreateJsonRpcError(-32700, TEXT("Empty request body"), nullptr);
+		FString ResponseStr = SerializeJson(ErrorResponse);
+		
+		TUniquePtr<FHttpServerResponse> Response = FHttpServerResponse::Create(ResponseStr, TEXT("application/json"));
+		Response->Headers.Add(TEXT("Access-Control-Allow-Origin"), { TEXT("*") });
+		Response->Code = EHttpServerResponseCodes::BadRequest;
+		OnComplete(MoveTemp(Response));
+		return true;
+	}
+	
+	// Parse JSON
+	TSharedRef<TJsonReader<>> Reader = TJsonReaderFactory<>::Create(RequestBody);
+	TSharedPtr<FJsonValue> ParsedJson;
+	
+	if (!FJsonSerializer::Deserialize(Reader, ParsedJson) || !ParsedJson.IsValid())
+	{
+		UE_LOG(LogTemp, Error, TEXT("[ECABridge] Failed to parse JSON: %s"), *RequestBody.Left(200));
+		TSharedPtr<FJsonObject> ErrorResponse = CreateJsonRpcError(-32700, TEXT("Parse error"), nullptr);
+		FString ResponseStr = SerializeJson(ErrorResponse);
+		
+		TUniquePtr<FHttpServerResponse> Response = FHttpServerResponse::Create(ResponseStr, TEXT("application/json"));
+		Response->Headers.Add(TEXT("Access-Control-Allow-Origin"), { TEXT("*") });
+		Response->Code = EHttpServerResponseCodes::BadRequest;
+		OnComplete(MoveTemp(Response));
+		return true;
+	}
+	
+	// Handle batch or single request
+	FString ResponseStr;
+	
+	if (ParsedJson->Type == EJson::Array)
+	{
+		// Batch request
+		TArray<TSharedPtr<FJsonValue>> Requests = ParsedJson->AsArray();
+		TArray<TSharedPtr<FJsonValue>> Responses;
+		
+		for (const TSharedPtr<FJsonValue>& ReqValue : Requests)
+		{
+			if (ReqValue->Type == EJson::Object)
+			{
+				TSharedPtr<FJsonObject> Result = ProcessJsonRpcRequest(ReqValue->AsObject());
+				if (Result.IsValid())
+				{
+					Responses.Add(MakeShared<FJsonValueObject>(Result));
+				}
+			}
+		}
+		
+		// Serialize array response
+		TSharedRef<TJsonWriter<>> Writer = TJsonWriterFactory<>::Create(&ResponseStr);
+		FJsonSerializer::Serialize(Responses, Writer);
+	}
+	else if (ParsedJson->Type == EJson::Object)
+	{
+		// Single request
+		TSharedPtr<FJsonObject> Result = ProcessJsonRpcRequest(ParsedJson->AsObject());
+		if (Result.IsValid())
+		{
+			ResponseStr = SerializeJson(Result);
+		}
+	}
+	
+	// If no response (notification), return 202 Accepted
+	if (ResponseStr.IsEmpty())
+	{
+		TUniquePtr<FHttpServerResponse> Response = FHttpServerResponse::Create(FString(), TEXT("application/json"));
+		Response->Headers.Add(TEXT("Access-Control-Allow-Origin"), { TEXT("*") });
+		Response->Code = EHttpServerResponseCodes::Accepted;
+		OnComplete(MoveTemp(Response));
+		return true;
+	}
+	
+	UE_LOG(LogTemp, Log, TEXT("[ECABridge] MCP Response: %s"), *ResponseStr.Left(500));
+	
+	TUniquePtr<FHttpServerResponse> Response = FHttpServerResponse::Create(ResponseStr, TEXT("application/json"));
+	Response->Headers.Add(TEXT("Access-Control-Allow-Origin"), { TEXT("*") });
+	OnComplete(MoveTemp(Response));
+	return true;
+}
+
+TSharedPtr<FJsonObject> FECAMCPServer::ProcessJsonRpcRequest(const TSharedPtr<FJsonObject>& Request)
+{
+	if (!Request.IsValid())
+	{
+		return CreateJsonRpcError(-32600, TEXT("Invalid Request"), nullptr);
+	}
+	
+	// Extract method
+	FString Method;
+	if (!Request->TryGetStringField(TEXT("method"), Method))
+	{
+		return CreateJsonRpcError(-32600, TEXT("Missing method field"), nullptr);
+	}
+	
+	// Extract id (can be string, number, or null; missing = notification)
+	TSharedPtr<FJsonValue> RequestId = Request->TryGetField(TEXT("id"));
+	bool bIsNotification = !RequestId.IsValid() || RequestId->IsNull();
+	
+	// Extract params
+	TSharedPtr<FJsonObject> Params;
+	const TSharedPtr<FJsonObject>* ParamsPtr;
+	if (Request->TryGetObjectField(TEXT("params"), ParamsPtr))
+	{
+		Params = *ParamsPtr;
+	}
+	
+	UE_LOG(LogTemp, Log, TEXT("[ECABridge] Processing: method=%s, notification=%s"), *Method, bIsNotification ? TEXT("true") : TEXT("false"));
+	
+	// Route to handler
+	TSharedPtr<FJsonObject> Result;
+	
+	if (Method == TEXT("initialize"))
+	{
+		Result = HandleInitialize(Params);
+	}
+	else if (Method == TEXT("tools/list"))
+	{
+		Result = HandleToolsList();
+	}
+	else if (Method == TEXT("tools/call"))
+	{
+		Result = HandleToolsCall(Params);
+	}
+	else if (Method == TEXT("notifications/initialized"))
+	{
+		// Client notification - no response needed
+		UE_LOG(LogTemp, Log, TEXT("[ECABridge] Client initialized notification received"));
+		return nullptr;
+	}
+	else if (Method == TEXT("ping"))
+	{
+		Result = MakeShared<FJsonObject>();
+	}
+	else
+	{
+		return CreateJsonRpcError(-32601, FString::Printf(TEXT("Method not found: %s"), *Method), RequestId);
+	}
+	
+	// Notifications don't get responses
+	if (bIsNotification)
+	{
+		return nullptr;
+	}
+	
+	return CreateJsonRpcResponse(Result, RequestId);
+}
+
+bool FECAMCPServer::HandleHealth(const FHttpServerRequest& Request, const FHttpResultCallback& OnComplete)
+{
+	TSharedPtr<FJsonObject> HealthJson = MakeShared<FJsonObject>();
+	HealthJson->SetStringField(TEXT("status"), TEXT("ok"));
+	HealthJson->SetStringField(TEXT("server"), ServerName);
+	HealthJson->SetStringField(TEXT("version"), ServerVersion);
+	HealthJson->SetStringField(TEXT("protocol"), ProtocolVersion);
+	HealthJson->SetNumberField(TEXT("commands"), FECACommandRegistry::Get().GetAllCommands().Num());
+	HealthJson->SetBoolField(TEXT("bridge_ready"), Bridge != nullptr);
+
+	FString ResponseBody = SerializeJson(HealthJson);
+
+	TUniquePtr<FHttpServerResponse> Response = FHttpServerResponse::Create(ResponseBody, TEXT("application/json"));
+	Response->Headers.Add(TEXT("Access-Control-Allow-Origin"), { TEXT("*") });
+	OnComplete(MoveTemp(Response));
+
+	return true;
+}
+
+bool FECAMCPServer::HandleCORSPreflight(const FHttpServerRequest& Request, const FHttpResultCallback& OnComplete)
+{
+	TUniquePtr<FHttpServerResponse> Response = FHttpServerResponse::Create(FString(), TEXT("text/plain"));
+	Response->Code = EHttpServerResponseCodes::NoContent;
+	
+	Response->Headers.Add(TEXT("Access-Control-Allow-Origin"), { TEXT("*") });
+	Response->Headers.Add(TEXT("Access-Control-Allow-Methods"), { TEXT("GET, POST, OPTIONS") });
+	Response->Headers.Add(TEXT("Access-Control-Allow-Headers"), { TEXT("Content-Type, Accept, Cache-Control") });
+	Response->Headers.Add(TEXT("Access-Control-Max-Age"), { TEXT("86400") });
+	
+	OnComplete(MoveTemp(Response));
+	return true;
+}
+
+TSharedPtr<FJsonObject> FECAMCPServer::HandleInitialize(const TSharedPtr<FJsonObject>& Params)
+{
+	TSharedPtr<FJsonObject> Result = MakeShared<FJsonObject>();
+	
+	// Server info
+	TSharedPtr<FJsonObject> ServerInfo = MakeShared<FJsonObject>();
+	ServerInfo->SetStringField(TEXT("name"), ServerName);
+	ServerInfo->SetStringField(TEXT("version"), ServerVersion);
+	Result->SetObjectField(TEXT("serverInfo"), ServerInfo);
+
+	// Protocol version
+	Result->SetStringField(TEXT("protocolVersion"), ProtocolVersion);
+
+	// Capabilities
+	TSharedPtr<FJsonObject> Capabilities = MakeShared<FJsonObject>();
+	
+	// Tools capability
+	TSharedPtr<FJsonObject> ToolsCap = MakeShared<FJsonObject>();
+	ToolsCap->SetBoolField(TEXT("listChanged"), false);
+	Capabilities->SetObjectField(TEXT("tools"), ToolsCap);
+
+	Result->SetObjectField(TEXT("capabilities"), Capabilities);
+
+	UE_LOG(LogTemp, Log, TEXT("[ECABridge] MCP client initialized"));
+
+	return Result;
+}
+
+TSharedPtr<FJsonObject> FECAMCPServer::HandleToolsList()
+{
+	TSharedPtr<FJsonObject> Result = MakeShared<FJsonObject>();
+	Result->SetArrayField(TEXT("tools"), BuildToolDefinitions());
+	
+	UE_LOG(LogTemp, Log, TEXT("[ECABridge] Returned %d tools"), FECACommandRegistry::Get().GetAllCommands().Num());
+	
+	return Result;
+}
+
+TSharedPtr<FJsonObject> FECAMCPServer::HandleToolsCall(const TSharedPtr<FJsonObject>& Params)
+{
+	if (!Params.IsValid())
+	{
+		TSharedPtr<FJsonObject> Error = MakeShared<FJsonObject>();
+		Error->SetBoolField(TEXT("isError"), true);
+		
+		TArray<TSharedPtr<FJsonValue>> Content;
+		TSharedPtr<FJsonObject> TextContent = MakeShared<FJsonObject>();
+		TextContent->SetStringField(TEXT("type"), TEXT("text"));
+		TextContent->SetStringField(TEXT("text"), TEXT("Missing parameters"));
+		Content.Add(MakeShared<FJsonValueObject>(TextContent));
+		Error->SetArrayField(TEXT("content"), Content);
+		
+		return Error;
+	}
+
+	FString ToolName;
+	if (!Params->TryGetStringField(TEXT("name"), ToolName))
+	{
+		TSharedPtr<FJsonObject> Error = MakeShared<FJsonObject>();
+		Error->SetBoolField(TEXT("isError"), true);
+		
+		TArray<TSharedPtr<FJsonValue>> Content;
+		TSharedPtr<FJsonObject> TextContent = MakeShared<FJsonObject>();
+		TextContent->SetStringField(TEXT("type"), TEXT("text"));
+		TextContent->SetStringField(TEXT("text"), TEXT("Missing tool name"));
+		Content.Add(MakeShared<FJsonValueObject>(TextContent));
+		Error->SetArrayField(TEXT("content"), Content);
+		
+		return Error;
+	}
+	
+	TSharedPtr<FJsonObject> Arguments;
+	const TSharedPtr<FJsonObject>* ArgsPtr;
+	if (Params->TryGetObjectField(TEXT("arguments"), ArgsPtr))
+	{
+		Arguments = *ArgsPtr;
+	}
+	else
+	{
+		Arguments = MakeShared<FJsonObject>();
+	}
+
+	UE_LOG(LogTemp, Log, TEXT("[ECABridge] Tool call: %s"), *ToolName);
+
+	// Execute the command
+	FECACommandResult CommandResult = FECACommandRegistry::Get().ExecuteCommand(ToolName, Arguments);
+
+	// Build MCP response
+	TSharedPtr<FJsonObject> Result = MakeShared<FJsonObject>();
+	
+	TArray<TSharedPtr<FJsonValue>> Content;
+	TSharedPtr<FJsonObject> TextContent = MakeShared<FJsonObject>();
+	TextContent->SetStringField(TEXT("type"), TEXT("text"));
+	
+	if (CommandResult.bSuccess)
+	{
+		FString ResultText;
+		if (CommandResult.ResultData.IsValid())
+		{
+			ResultText = SerializeJson(CommandResult.ResultData);
+		}
+		else
+		{
+			ResultText = TEXT("{\"success\": true}");
+		}
+		TextContent->SetStringField(TEXT("text"), ResultText);
+	}
+	else
+	{
+		Result->SetBoolField(TEXT("isError"), true);
+		TextContent->SetStringField(TEXT("text"), CommandResult.ErrorMessage);
+	}
+	
+	Content.Add(MakeShared<FJsonValueObject>(TextContent));
+	Result->SetArrayField(TEXT("content"), Content);
+
+	return Result;
+}
+
+TArray<TSharedPtr<FJsonValue>> FECAMCPServer::BuildToolDefinitions()
+{
+	TArray<TSharedPtr<FJsonValue>> Tools;
+
+	for (const TSharedPtr<IECACommand>& Command : FECACommandRegistry::Get().GetAllCommands())
+	{
+		TSharedPtr<FJsonObject> ToolDef = MakeShared<FJsonObject>();
+		ToolDef->SetStringField(TEXT("name"), Command->GetName());
+		ToolDef->SetStringField(TEXT("description"), Command->GetDescription());
+
+		// Build input schema
+		TSharedPtr<FJsonObject> InputSchema = MakeShared<FJsonObject>();
+		InputSchema->SetStringField(TEXT("type"), TEXT("object"));
+
+		TSharedPtr<FJsonObject> Properties = MakeShared<FJsonObject>();
+		TArray<TSharedPtr<FJsonValue>> Required;
+
+		for (const FECACommandParam& Param : Command->GetParameters())
+		{
+			TSharedPtr<FJsonObject> PropDef = MakeShared<FJsonObject>();
+			
+			// Map parameter types to valid JSON Schema types
+			// Valid types: string, number, integer, boolean, object, array, null
+			FString JsonSchemaType = Param.Type;
+			if (Param.Type == TEXT("any"))
+			{
+				// "any" is not valid in JSON Schema - use a union or just string
+				// We'll accept any JSON value by not specifying type constraints
+				PropDef->SetStringField(TEXT("description"), Param.Description + TEXT(" (accepts any JSON value)"));
+			}
+			else if (Param.Type == TEXT("vector") || Param.Type == TEXT("rotator") || Param.Type == TEXT("transform"))
+			{
+				// These are objects with specific structure
+				JsonSchemaType = TEXT("object");
+				PropDef->SetStringField(TEXT("type"), JsonSchemaType);
+				PropDef->SetStringField(TEXT("description"), Param.Description);
+			}
+			else if (Param.Type == TEXT("float") || Param.Type == TEXT("double"))
+			{
+				JsonSchemaType = TEXT("number");
+				PropDef->SetStringField(TEXT("type"), JsonSchemaType);
+				PropDef->SetStringField(TEXT("description"), Param.Description);
+			}
+			else if (Param.Type == TEXT("int") || Param.Type == TEXT("int32") || Param.Type == TEXT("int64"))
+			{
+				JsonSchemaType = TEXT("integer");
+				PropDef->SetStringField(TEXT("type"), JsonSchemaType);
+				PropDef->SetStringField(TEXT("description"), Param.Description);
+			}
+			else
+			{
+				// Standard JSON Schema types: string, number, integer, boolean, object, array
+				PropDef->SetStringField(TEXT("type"), JsonSchemaType);
+				PropDef->SetStringField(TEXT("description"), Param.Description);
+			}
+			
+			if (!Param.DefaultValue.IsEmpty())
+			{
+				PropDef->SetStringField(TEXT("default"), Param.DefaultValue);
+			}
+
+			Properties->SetObjectField(Param.Name, PropDef);
+
+			if (Param.bRequired)
+			{
+				Required.Add(MakeShared<FJsonValueString>(Param.Name));
+			}
+		}
+
+		InputSchema->SetObjectField(TEXT("properties"), Properties);
+		if (Required.Num() > 0)
+		{
+			InputSchema->SetArrayField(TEXT("required"), Required);
+		}
+
+		ToolDef->SetObjectField(TEXT("inputSchema"), InputSchema);
+		Tools.Add(MakeShared<FJsonValueObject>(ToolDef));
+	}
+
+	return Tools;
+}
+
+TSharedPtr<FJsonObject> FECAMCPServer::CreateJsonRpcResponse(const TSharedPtr<FJsonObject>& Result, const TSharedPtr<FJsonValue>& RequestId)
+{
+	TSharedPtr<FJsonObject> Response = MakeShared<FJsonObject>();
+	Response->SetStringField(TEXT("jsonrpc"), TEXT("2.0"));
+	
+	if (RequestId.IsValid())
+	{
+		Response->SetField(TEXT("id"), RequestId);
+	}
+	
+	Response->SetObjectField(TEXT("result"), Result.IsValid() ? Result : MakeShared<FJsonObject>());
+	
+	return Response;
+}
+
+TSharedPtr<FJsonObject> FECAMCPServer::CreateJsonRpcError(int32 Code, const FString& Message, const TSharedPtr<FJsonValue>& RequestId)
+{
+	TSharedPtr<FJsonObject> Response = MakeShared<FJsonObject>();
+	Response->SetStringField(TEXT("jsonrpc"), TEXT("2.0"));
+	
+	if (RequestId.IsValid())
+	{
+		Response->SetField(TEXT("id"), RequestId);
+	}
+	else
+	{
+		Response->SetField(TEXT("id"), MakeShared<FJsonValueNull>());
+	}
+
+	TSharedPtr<FJsonObject> Error = MakeShared<FJsonObject>();
+	Error->SetNumberField(TEXT("code"), Code);
+	Error->SetStringField(TEXT("message"), Message);
+	Response->SetObjectField(TEXT("error"), Error);
+
+	return Response;
+}
+
+FString FECAMCPServer::SerializeJson(const TSharedPtr<FJsonObject>& JsonObject)
+{
+	FString OutputString;
+	TSharedRef<TJsonWriter<>> Writer = TJsonWriterFactory<>::Create(&OutputString);
+	FJsonSerializer::Serialize(JsonObject.ToSharedRef(), Writer);
+	return OutputString;
+}
diff --git a/Engine/Plugins/Experimental/ECABridge/Source/ECABridge/Private/ECAServerRunnable.cpp b/Engine/Plugins/Experimental/ECABridge/Source/ECABridge/Private/ECAServerRunnable.cpp
new file mode 100644
--- /dev/null
+++ b/Engine/Plugins/Experimental/ECABridge/Source/ECABridge/Private/ECAServerRunnable.cpp
@@ -0,0 +1,239 @@
+// Copyright Epic Games, Inc. All Rights Reserved.
+
+#include "ECAServerRunnable.h"
+#include "ECABridge.h"
+#include "ECABridgeSettings.h"
+#include "Sockets.h"
+#include "SocketSubsystem.h"
+#include "Async/Async.h"
+
+FECAServerRunnable::FECAServerRunnable(UECABridge* InBridge, int32 InPort)
+	: Bridge(InBridge)
+	, ListenerSocket(nullptr)
+	, Port(InPort)
+	, bShouldStop(false)
+{
+}
+
+FECAServerRunnable::~FECAServerRunnable()
+{
+	if (ListenerSocket)
+	{
+		ListenerSocket->Close();
+		ISocketSubsystem::Get(PLATFORM_SOCKETSUBSYSTEM)->DestroySocket(ListenerSocket);
+		ListenerSocket = nullptr;
+	}
+}
+
+bool FECAServerRunnable::Init()
+{
+	ISocketSubsystem* SocketSubsystem = ISocketSubsystem::Get(PLATFORM_SOCKETSUBSYSTEM);
+	if (!SocketSubsystem)
+	{
+		UE_LOG(LogTemp, Error, TEXT("[ECABridge] Failed to get socket subsystem"));
+		return false;
+	}
+	
+	// Create listener socket
+	ListenerSocket = SocketSubsystem->CreateSocket(NAME_Stream, TEXT("ECAServer"), false);
+	if (!ListenerSocket)
+	{
+		UE_LOG(LogTemp, Error, TEXT("[ECABridge] Failed to create listener socket"));
+		return false;
+	}
+	
+	// Configure socket
+	ListenerSocket->SetReuseAddr(true);
+	ListenerSocket->SetNonBlocking(true);
+	ListenerSocket->SetNoDelay(true);
+	
+	// Bind to address
+	TSharedRef<FInternetAddr> Addr = SocketSubsystem->CreateInternetAddr();
+	
+	const UECABridgeSettings* Settings = UECABridgeSettings::Get();
+	if (Settings->bAllowRemoteConnections)
+	{
+		Addr->SetAnyAddress();
+		UE_LOG(LogTemp, Warning, TEXT("[ECABridge] Allowing remote connections (security risk!)"));
+	}
+	else
+	{
+		// Localhost only
+		bool bIsValid;
+		Addr->SetIp(TEXT("127.0.0.1"), bIsValid);
+	}
+	Addr->SetPort(Port);
+	
+	if (!ListenerSocket->Bind(*Addr))
+	{
+		UE_LOG(LogTemp, Error, TEXT("[ECABridge] Failed to bind to port %d"), Port);
+		return false;
+	}
+	
+	if (!ListenerSocket->Listen(8))
+	{
+		UE_LOG(LogTemp, Error, TEXT("[ECABridge] Failed to listen on socket"));
+		return false;
+	}
+	
+	UE_LOG(LogTemp, Log, TEXT("[ECABridge] TCP server listening on port %d"), Port);
+	return true;
+}
+
+uint32 FECAServerRunnable::Run()
+{
+	while (!bShouldStop)
+	{
+		bool bHasPendingConnection = false;
+		if (ListenerSocket->HasPendingConnection(bHasPendingConnection) && bHasPendingConnection)
+		{
+			FSocket* ClientSocket = ListenerSocket->Accept(TEXT("ECAClient"));
+			if (ClientSocket)
+			{
+				HandleClient(ClientSocket);
+				ClientSocket->Close();
+				ISocketSubsystem::Get(PLATFORM_SOCKETSUBSYSTEM)->DestroySocket(ClientSocket);
+			}
+		}
+		
+		// Sleep to avoid spinning
+		FPlatformProcess::Sleep(0.01f);
+	}
+	
+	return 0;
+}
+
+void FECAServerRunnable::Stop()
+{
+	bShouldStop = true;
+}
+
+void FECAServerRunnable::HandleClient(FSocket* ClientSocket)
+{
+	ClientSocket->SetNonBlocking(false);
+	
+	// Read the command
+	FString JsonCommand;
+	if (!ReadMessage(ClientSocket, JsonCommand))
+	{
+		return;
+	}
+	
+	if (JsonCommand.IsEmpty())
+	{
+		return;
+	}
+	
+	// Execute on game thread and wait for result
+	FString Response;
+	FEvent* CompletionEvent = FPlatformProcess::GetSynchEventFromPool();
+	
+	AsyncTask(ENamedThreads::GameThread, [this, &JsonCommand, &Response, CompletionEvent]()
+	{
+		if (Bridge)
+		{
+			Response = Bridge->ProcessCommand(JsonCommand);
+		}
+		else
+		{
+			Response = TEXT("{\"success\":false,\"error\":\"Bridge not available\"}");
+		}
+		CompletionEvent->Trigger();
+	});
+	
+	// Wait for completion with timeout
+	const UECABridgeSettings* Settings = UECABridgeSettings::Get();
+	bool bCompleted = CompletionEvent->Wait(Settings->CommandTimeoutSeconds * 1000);
+	FPlatformProcess::ReturnSynchEventToPool(CompletionEvent);
+	
+	if (!bCompleted)
+	{
+		Response = TEXT("{\"success\":false,\"error\":\"Command execution timed out\"}");
+	}
+	
+	// Send response
+	SendResponse(ClientSocket, Response);
+}
+
+bool FECAServerRunnable::ReadMessage(FSocket* ClientSocket, FString& OutMessage)
+{
+	TArray<uint8> Buffer;
+	Buffer.SetNumUninitialized(BufferSize);
+	
+	int32 BytesRead = 0;
+	
+	// First, try to read a length prefix (4 bytes, big-endian)
+	uint8 LengthBuffer[4];
+	int32 LengthBytesRead = 0;
+	
+	// Try to peek if we have a length-prefixed message
+	if (ClientSocket->Recv(LengthBuffer, 4, LengthBytesRead, ESocketReceiveFlags::Peek))
+	{
+		// Check if this looks like a length prefix (reasonable size)
+		uint32 MessageLength = (LengthBuffer[0] << 24) | (LengthBuffer[1] << 16) | (LengthBuffer[2] << 8) | LengthBuffer[3];
+		
+		if (MessageLength > 0 && MessageLength < BufferSize && LengthBytesRead == 4)
+		{
+			// Looks like a length-prefixed message, consume the length bytes
+			ClientSocket->Recv(LengthBuffer, 4, LengthBytesRead);
+			
+			// Now read exactly that many bytes
+			TArray<uint8> MessageBuffer;
+			MessageBuffer.SetNumUninitialized(MessageLength);
+			
+			int32 TotalRead = 0;
+			while (TotalRead < (int32)MessageLength)
+			{
+				int32 ChunkRead = 0;
+				if (!ClientSocket->Recv(MessageBuffer.GetData() + TotalRead, MessageLength - TotalRead, ChunkRead))
+				{
+					break;
+				}
+				TotalRead += ChunkRead;
+			}
+			
+			if (TotalRead == (int32)MessageLength)
+			{
+				OutMessage = FString(UTF8_TO_TCHAR(reinterpret_cast<const char*>(MessageBuffer.GetData())));
+				OutMessage = OutMessage.Left(MessageLength);
+				return true;
+			}
+		}
+	}
+	
+	// Fallback: Read until we get valid JSON (for simple clients)
+	if (ClientSocket->Recv(Buffer.GetData(), Buffer.Num(), BytesRead))
+	{
+		if (BytesRead > 0)
+		{
+			OutMessage = FString(UTF8_TO_TCHAR(reinterpret_cast<const char*>(Buffer.GetData())));
+			OutMessage = OutMessage.Left(BytesRead);
+			return true;
+		}
+	}
+	
+	return false;
+}
+
+bool FECAServerRunnable::SendResponse(FSocket* ClientSocket, const FString& Response)
+{
+	FTCHARToUTF8 ResponseUtf8(*Response);
+	const int32 ResponseLength = ResponseUtf8.Length();
+	
+	// Send length prefix (4 bytes, big-endian)
+	uint8 LengthBuffer[4];
+	LengthBuffer[0] = (ResponseLength >> 24) & 0xFF;
+	LengthBuffer[1] = (ResponseLength >> 16) & 0xFF;
+	LengthBuffer[2] = (ResponseLength >> 8) & 0xFF;
+	LengthBuffer[3] = ResponseLength & 0xFF;
+	
+	int32 BytesSent = 0;
+	if (!ClientSocket->Send(LengthBuffer, 4, BytesSent))
+	{
+		return false;
+	}
+	
+	// Send message body
+	BytesSent = 0;
+	return ClientSocket->Send(reinterpret_cast<const uint8*>(ResponseUtf8.Get()), ResponseLength, BytesSent);
+}
