ECABridge experimental MCP plugin

diff --git a/Engine/Plugins/Experimental/ECABridge/Source/ECABridge/Public/BlueprintLisp.h b/Engine/Plugins/Experimental/ECABridge/Source/ECABridge/Public/BlueprintLisp.h
new file mode 100644
--- /dev/null
+++ b/Engine/Plugins/Experimental/ECABridge/Source/ECABridge/Public/BlueprintLisp.h
@@ -0,0 +1,266 @@
+// Copyright Epic Games, Inc. All Rights Reserved.
+
+#pragma once
+
+#include "CoreMinimal.h"
+#include "Dom/JsonObject.h"
+
+/**
+ * BlueprintLisp - A LISP-like DSL for representing Blueprint graphs
+ * 
+ * This module provides bidirectional conversion between Blueprint graphs and a
+ * simple S-expression syntax that's easy for AI to read, understand, and generate.
+ * 
+ * Example syntax:
+ * 
+ *   (event BeginPlay
+ *     (let player (GetPlayerCharacter 0))
+ *     (branch (IsValid player)
+ *       :true (seq
+ *         (let health (call player GetHealth))
+ *         (PrintString (format "Health: {}" health)))
+ *       :false
+ *         (PrintString "No player!")))
+ * 
+ * Core forms:
+ *   (event Name ...)           - Event node
+ *   (func Name ...)            - Function definition  
+ *   (let var expr)             - Local variable
+ *   (set var expr)             - Set variable
+ *   (seq ...)                  - Sequence node (execution flow)
+ *   (branch cond :true :false) - Branch node
+ *   (foreach item coll ...)    - ForEach loop
+ *   (call target func args...) - Call function on object
+ *   (delay seconds)            - Delay node
+ *   (cast Type var body)       - Cast with success exec
+ *   (vec x y z)                - Vector literal
+ *   (rot p y r)                - Rotator literal
+ */
+
+//------------------------------------------------------------------------------
+// S-Expression AST
+//------------------------------------------------------------------------------
+
+/** Types of S-expression nodes */
+enum class ELispNodeType : uint8
+{
+	Nil,        // Empty/null
+	Symbol,     // Identifier: foo, BeginPlay, GetHealth
+	Keyword,    // Keyword: :true, :false, :pin-name
+	Number,     // Numeric literal: 42, 3.14
+	String,     // String literal: "hello"
+	List,       // List: (a b c)
+};
+
+/** Forward declaration */
+struct FLispNode;
+
+/** Shared pointer to a LISP node */
+using FLispNodePtr = TSharedPtr<FLispNode>;
+
+/** A node in the S-expression AST */
+struct ECABRIDGE_API FLispNode
+{
+	ELispNodeType Type = ELispNodeType::Nil;
+	
+	// Value storage (only one is used based on Type)
+	FString StringValue;      // For Symbol, Keyword, String
+	double NumberValue = 0.0; // For Number
+	TArray<FLispNodePtr> Children; // For List
+	
+	// Source location for error reporting
+	int32 Line = 0;
+	int32 Column = 0;
+	
+	// Constructors
+	FLispNode() = default;
+	
+	static FLispNodePtr MakeNil();
+	static FLispNodePtr MakeSymbol(const FString& Name);
+	static FLispNodePtr MakeKeyword(const FString& Name);
+	static FLispNodePtr MakeNumber(double Value);
+	static FLispNodePtr MakeString(const FString& Value);
+	static FLispNodePtr MakeList(const TArray<FLispNodePtr>& Items = {});
+	
+	// Type checks
+	bool IsNil() const { return Type == ELispNodeType::Nil; }
+	bool IsSymbol() const { return Type == ELispNodeType::Symbol; }
+	bool IsKeyword() const { return Type == ELispNodeType::Keyword; }
+	bool IsNumber() const { return Type == ELispNodeType::Number; }
+	bool IsString() const { return Type == ELispNodeType::String; }
+	bool IsList() const { return Type == ELispNodeType::List; }
+	
+	// Check if this is a list starting with a specific symbol
+	bool IsForm(const FString& FormName) const;
+	
+	// Get the form name (first element if list with symbol head)
+	FString GetFormName() const;
+	
+	// List accessors
+	int32 Num() const { return Children.Num(); }
+	FLispNodePtr operator[](int32 Index) const;
+	FLispNodePtr Get(int32 Index) const;
+	
+	// Find keyword argument in list: (foo :key value) -> returns value for :key
+	FLispNodePtr GetKeywordArg(const FString& Keyword) const;
+	bool HasKeyword(const FString& Keyword) const;
+	
+	// Convert to string representation
+	FString ToString(bool bPretty = false, int32 IndentLevel = 0) const;
+};
+
+//------------------------------------------------------------------------------
+// Parser
+//------------------------------------------------------------------------------
+
+/** Result of parsing */
+struct ECABRIDGE_API FLispParseResult
+{
+	bool bSuccess = false;
+	FString Error;
+	int32 ErrorLine = 0;
+	int32 ErrorColumn = 0;
+	TArray<FLispNodePtr> Nodes; // Top-level expressions
+	
+	static FLispParseResult Success(const TArray<FLispNodePtr>& Nodes);
+	static FLispParseResult Failure(const FString& Error, int32 Line = 0, int32 Column = 0);
+};
+
+/** S-expression parser */
+class ECABRIDGE_API FLispParser
+{
+public:
+	/** Parse a string containing one or more S-expressions */
+	static FLispParseResult Parse(const FString& Source);
+	
+private:
+	FLispParser(const FString& Source);
+	
+	FLispParseResult DoParse();
+	FLispNodePtr ParseExpr();
+	FLispNodePtr ParseList();
+	FLispNodePtr ParseAtom();
+	FLispNodePtr ParseString();
+	FLispNodePtr ParseNumber();
+	FLispNodePtr ParseSymbolOrKeyword();
+	
+	void SkipWhitespaceAndComments();
+	TCHAR Peek() const;
+	TCHAR Advance();
+	bool IsAtEnd() const;
+	bool Match(TCHAR Expected);
+	
+	FLispParseResult MakeError(const FString& Message);
+	
+	const FString& Source;
+	int32 Position = 0;
+	int32 Line = 1;
+	int32 Column = 1;
+};
+
+//------------------------------------------------------------------------------
+// Blueprint Graph Conversion
+//------------------------------------------------------------------------------
+
+/** Options for converting Blueprint to Lisp */
+struct ECABRIDGE_API FBlueprintToLispOptions
+{
+	bool bIncludeComments = true;
+	bool bIncludePositions = false; // Include node positions as metadata
+	bool bPrettyPrint = true;
+};
+
+/** Options for converting Lisp to Blueprint */
+struct ECABRIDGE_API FLispToBlueprintOptions
+{
+	bool bClearExisting = false; // Clear existing nodes before adding
+	bool bAutoLayout = true;     // Auto-arrange nodes after creation
+};
+
+/** Result of Blueprint conversion */
+struct ECABRIDGE_API FBlueprintLispResult
+{
+	bool bSuccess = false;
+	FString Error;
+	FString LispCode;                    // For BP -> Lisp
+	TSharedPtr<FJsonObject> GraphData;   // For Lisp -> BP (intermediate JSON)
+	TArray<FString> Warnings;
+	
+	static FBlueprintLispResult Success(const FString& Code);
+	static FBlueprintLispResult SuccessWithData(TSharedPtr<FJsonObject> Data);
+	static FBlueprintLispResult Failure(const FString& Error);
+};
+
+/**
+ * Main converter class for Blueprint <-> Lisp conversion
+ */
+class ECABRIDGE_API FBlueprintLispConverter
+{
+public:
+	/** Convert a Blueprint graph to Lisp code */
+	static FBlueprintLispResult BlueprintToLisp(
+		const FString& BlueprintPath,
+		const FString& GraphName = TEXT("EventGraph"),
+		const FBlueprintToLispOptions& Options = FBlueprintToLispOptions());
+	
+	/** Convert Lisp code to Blueprint graph operations (returns JSON for ECA commands) */
+	static FBlueprintLispResult LispToBlueprint(
+		const FString& LispCode,
+		const FString& BlueprintPath,
+		const FString& GraphName = TEXT("EventGraph"),
+		const FLispToBlueprintOptions& Options = FLispToBlueprintOptions());
+	
+	/** Validate Lisp code without generating Blueprint */
+	static FBlueprintLispResult ValidateLisp(const FString& LispCode);
+	
+private:
+	// Blueprint -> Lisp helpers
+	static FLispNodePtr NodeToLisp(const TSharedPtr<FJsonObject>& Node, 
+		const TMap<FString, TSharedPtr<FJsonObject>>& NodeMap,
+		TSet<FString>& VisitedNodes,
+		const FBlueprintToLispOptions& Options);
+	
+	static FLispNodePtr ExecChainToLisp(const FString& StartNodeId,
+		const TMap<FString, TSharedPtr<FJsonObject>>& NodeMap,
+		TSet<FString>& VisitedNodes,
+		const FBlueprintToLispOptions& Options);
+	
+	static FLispNodePtr PureNodeToLisp(const TSharedPtr<FJsonObject>& Node,
+		const TMap<FString, TSharedPtr<FJsonObject>>& NodeMap,
+		TSet<FString>& VisitedNodes);
+	
+	// Lisp -> Blueprint helpers  
+	static bool ProcessForm(const FLispNodePtr& Form,
+		TSharedPtr<FJsonObject> GraphData,
+		TMap<FString, FString>& VariableMap,
+		FString& OutError);
+	
+	static TSharedPtr<FJsonObject> CreateNodeJson(const FString& NodeType,
+		const TMap<FString, FString>& Properties = {});
+	
+	static bool ResolveExpression(const FLispNodePtr& Expr,
+		TSharedPtr<FJsonObject> GraphData,
+		TMap<FString, FString>& VariableMap,
+		FString& OutNodeId,
+		FString& OutPinName,
+		FString& OutError);
+};
+
+//------------------------------------------------------------------------------
+// Utility Functions
+//------------------------------------------------------------------------------
+
+namespace BlueprintLisp
+{
+	/** Format Lisp code with proper indentation */
+	ECABRIDGE_API FString PrettyPrint(const FString& LispCode);
+	
+	/** Minify Lisp code (remove extra whitespace and comments) */
+	ECABRIDGE_API FString Minify(const FString& LispCode);
+	
+	/** Extract all symbol names from Lisp code */
+	ECABRIDGE_API TArray<FString> ExtractSymbols(const FString& LispCode);
+	
+	/** Check if a string is a valid symbol name */
+	ECABRIDGE_API bool IsValidSymbol(const FString& Name);
+}
diff --git a/Engine/Plugins/Experimental/ECABridge/Source/ECABridge/Public/Commands/ECAActorCommands.h b/Engine/Plugins/Experimental/ECABridge/Source/ECABridge/Public/Commands/ECAActorCommands.h
new file mode 100644
--- /dev/null
+++ b/Engine/Plugins/Experimental/ECABridge/Source/ECABridge/Public/Commands/ECAActorCommands.h
@@ -0,0 +1,188 @@
+// Copyright Epic Games, Inc. All Rights Reserved.
+
+#pragma once
+
+#include "Commands/ECACommand.h"
+
+/**
+ * Get all actors in the current level
+ */
+class FECACommand_GetActorsInLevel : public IECACommand
+{
+public:
+	virtual FString GetName() const override { return TEXT("get_actors_in_level"); }
+	virtual FString GetDescription() const override { return TEXT("Get all actors in the current level with optional class filtering"); }
+	virtual FString GetCategory() const override { return TEXT("Actor"); }
+	
+	virtual TArray<FECACommandParam> GetParameters() const override
+	{
+		return {
+			{ TEXT("class_filter"), TEXT("string"), TEXT("Filter actors by class name (partial match)"), false },
+			{ TEXT("include_hidden"), TEXT("boolean"), TEXT("Include hidden actors"), false, TEXT("false") }
+		};
+	}
+	
+	virtual FECACommandResult Execute(const TSharedPtr<FJsonObject>& Params) override;
+};
+
+/**
+ * Create a new actor in the level
+ */
+class FECACommand_CreateActor : public IECACommand
+{
+public:
+	virtual FString GetName() const override { return TEXT("create_actor"); }
+	virtual FString GetDescription() const override { return TEXT("Create a new actor in the current level"); }
+	virtual FString GetCategory() const override { return TEXT("Actor"); }
+	
+	virtual TArray<FECACommandParam> GetParameters() const override
+	{
+		return {
+			{ TEXT("actor_type"), TEXT("string"), TEXT("Type of actor (StaticMeshActor, PointLight, SpotLight, DirectionalLight, CameraActor, or class path)"), true },
+			{ TEXT("name"), TEXT("string"), TEXT("Display name for the actor"), false },
+			{ TEXT("location"), TEXT("object"), TEXT("Spawn location {x, y, z}"), false },
+			{ TEXT("rotation"), TEXT("object"), TEXT("Spawn rotation {pitch, yaw, roll}"), false },
+			{ TEXT("scale"), TEXT("object"), TEXT("Scale {x, y, z}"), false },
+			{ TEXT("mesh"), TEXT("string"), TEXT("Mesh path for StaticMeshActor"), false },
+			{ TEXT("folder"), TEXT("string"), TEXT("Outliner folder path"), false }
+		};
+	}
+	
+	virtual FECACommandResult Execute(const TSharedPtr<FJsonObject>& Params) override;
+};
+
+/**
+ * Delete an actor from the level
+ */
+class FECACommand_DeleteActor : public IECACommand
+{
+public:
+	virtual FString GetName() const override { return TEXT("delete_actor"); }
+	virtual FString GetDescription() const override { return TEXT("Delete an actor from the current level"); }
+	virtual FString GetCategory() const override { return TEXT("Actor"); }
+	
+	virtual TArray<FECACommandParam> GetParameters() const override
+	{
+		return {
+			{ TEXT("actor_name"), TEXT("string"), TEXT("Name of the actor to delete"), true }
+		};
+	}
+	
+	virtual FECACommandResult Execute(const TSharedPtr<FJsonObject>& Params) override;
+};
+
+/**
+ * Set actor transform (location, rotation, scale)
+ */
+class FECACommand_SetActorTransform : public IECACommand
+{
+public:
+	virtual FString GetName() const override { return TEXT("set_actor_transform"); }
+	virtual FString GetDescription() const override { return TEXT("Set the transform of an actor"); }
+	virtual FString GetCategory() const override { return TEXT("Actor"); }
+	
+	virtual TArray<FECACommandParam> GetParameters() const override
+	{
+		return {
+			{ TEXT("actor_name"), TEXT("string"), TEXT("Name of the actor"), true },
+			{ TEXT("location"), TEXT("object"), TEXT("New location {x, y, z}"), false },
+			{ TEXT("rotation"), TEXT("object"), TEXT("New rotation {pitch, yaw, roll}"), false },
+			{ TEXT("scale"), TEXT("object"), TEXT("New scale {x, y, z}"), false }
+		};
+	}
+	
+	virtual FECACommandResult Execute(const TSharedPtr<FJsonObject>& Params) override;
+};
+
+/**
+ * Get actor properties
+ */
+class FECACommand_GetActorProperties : public IECACommand
+{
+public:
+	virtual FString GetName() const override { return TEXT("get_actor_properties"); }
+	virtual FString GetDescription() const override { return TEXT("Get properties of an actor"); }
+	virtual FString GetCategory() const override { return TEXT("Actor"); }
+	
+	virtual TArray<FECACommandParam> GetParameters() const override
+	{
+		return {
+			{ TEXT("actor_name"), TEXT("string"), TEXT("Name of the actor"), true },
+			{ TEXT("properties"), TEXT("array"), TEXT("Specific property names to retrieve (optional, returns all if not specified)"), false }
+		};
+	}
+	
+	virtual FECACommandResult Execute(const TSharedPtr<FJsonObject>& Params) override;
+};
+
+/**
+ * Set actor property
+ */
+class FECACommand_SetActorProperty : public IECACommand
+{
+public:
+	virtual FString GetName() const override { return TEXT("set_actor_property"); }
+	virtual FString GetDescription() const override { return TEXT("Set a property on an actor. Supports StaticMesh, Material, and generic UE properties via reflection."); }
+	virtual FString GetCategory() const override { return TEXT("Actor"); }
+	
+	virtual TArray<FECACommandParam> GetParameters() const override
+	{
+		return {
+			{ TEXT("actor_path"), TEXT("string"), TEXT("Path to the actor (from get_selected_actors). Preferred over actor_name."), false },
+			{ TEXT("actor_name"), TEXT("string"), TEXT("Name of the actor (legacy, use actor_path instead)"), false },
+			{ TEXT("property_name"), TEXT("string"), TEXT("Property: 'StaticMesh', 'Material', 'Material[0]', 'Mobility', 'Hidden', 'Folder', or any UProperty name"), true },
+			{ TEXT("property_value"), TEXT("any"), TEXT("Value to set (asset path for objects, or primitive value)"), true },
+			{ TEXT("component_name"), TEXT("string"), TEXT("Optional: target a specific component instead of the actor"), false }
+		};
+	}
+	
+	virtual FECACommandResult Execute(const TSharedPtr<FJsonObject>& Params) override;
+};
+
+/**
+ * Find actors by various criteria
+ */
+class FECACommand_FindActors : public IECACommand
+{
+public:
+	virtual FString GetName() const override { return TEXT("find_actors"); }
+	virtual FString GetDescription() const override { return TEXT("Find actors by name pattern, class, tag, or location"); }
+	virtual FString GetCategory() const override { return TEXT("Actor"); }
+	
+	virtual TArray<FECACommandParam> GetParameters() const override
+	{
+		return {
+			{ TEXT("name_pattern"), TEXT("string"), TEXT("Name pattern to match (supports * wildcard)"), false },
+			{ TEXT("class_name"), TEXT("string"), TEXT("Class name to filter by"), false },
+			{ TEXT("tag"), TEXT("string"), TEXT("Actor tag to filter by"), false },
+			{ TEXT("near_location"), TEXT("object"), TEXT("Find actors near this location {x, y, z}"), false },
+			{ TEXT("radius"), TEXT("number"), TEXT("Radius for location search"), false, TEXT("1000") }
+		};
+	}
+	
+	virtual FECACommandResult Execute(const TSharedPtr<FJsonObject>& Params) override;
+};
+
+/**
+ * Duplicate an actor
+ */
+class FECACommand_DuplicateActor : public IECACommand
+{
+public:
+	virtual FString GetName() const override { return TEXT("duplicate_actor"); }
+	virtual FString GetDescription() const override { return TEXT("Duplicate an existing actor"); }
+	virtual FString GetCategory() const override { return TEXT("Actor"); }
+	
+	virtual TArray<FECACommandParam> GetParameters() const override
+	{
+		return {
+			{ TEXT("actor_name"), TEXT("string"), TEXT("Name of the actor to duplicate"), true },
+			{ TEXT("new_name"), TEXT("string"), TEXT("Name for the duplicated actor"), false },
+			{ TEXT("offset"), TEXT("object"), TEXT("Position offset from original {x, y, z}"), false }
+		};
+	}
+	
+	virtual FECACommandResult Execute(const TSharedPtr<FJsonObject>& Params) override;
+};
+
+
diff --git a/Engine/Plugins/Experimental/ECABridge/Source/ECABridge/Public/Commands/ECAAssetCommands.h b/Engine/Plugins/Experimental/ECABridge/Source/ECABridge/Public/Commands/ECAAssetCommands.h
new file mode 100644
--- /dev/null
+++ b/Engine/Plugins/Experimental/ECABridge/Source/ECABridge/Public/Commands/ECAAssetCommands.h
@@ -0,0 +1,650 @@
+// Copyright Epic Games, Inc. All Rights Reserved.
+
+#pragma once
+
+#include "Commands/ECACommand.h"
+
+/**
+ * Import a texture from a file (PNG, JPG, TGA, etc.)
+ */
+class FECACommand_ImportTexture : public IECACommand
+{
+public:
+	virtual FString GetName() const override { return TEXT("import_texture"); }
+	virtual FString GetDescription() const override { return TEXT("Import a texture from an image file (PNG, JPG, TGA, BMP, EXR) into a UAsset"); }
+	virtual FString GetCategory() const override { return TEXT("Asset"); }
+	
+	virtual TArray<FECACommandParam> GetParameters() const override
+	{
+		return {
+			{ TEXT("source_path"), TEXT("string"), TEXT("File path to the source image (PNG, JPG, TGA, etc.)"), true },
+			{ TEXT("destination_path"), TEXT("string"), TEXT("Content path for the imported texture (e.g., /Game/Textures/MyTexture)"), true },
+			{ TEXT("texture_name"), TEXT("string"), TEXT("Name for the texture asset (defaults to source filename)"), false },
+			{ TEXT("compression_settings"), TEXT("string"), TEXT("Compression settings (Default, Normalmap, Masks, Grayscale, HDR, UserInterface2D)"), false, TEXT("Default") },
+			{ TEXT("srgb"), TEXT("boolean"), TEXT("Enable sRGB color space"), false, TEXT("true") },
+			{ TEXT("lod_group"), TEXT("string"), TEXT("LOD group (World, WorldNormalMap, WorldSpecular, UI, Shadowmap, etc.)"), false, TEXT("World") }
+		};
+	}
+	
+	virtual FECACommandResult Execute(const TSharedPtr<FJsonObject>& Params) override;
+};
+
+/**
+ * Create a new material asset
+ */
+class FECACommand_CreateMaterial : public IECACommand
+{
+public:
+	virtual FString GetName() const override { return TEXT("create_material"); }
+	virtual FString GetDescription() const override { return TEXT("Create a new material asset"); }
+	virtual FString GetCategory() const override { return TEXT("Asset"); }
+	
+	virtual TArray<FECACommandParam> GetParameters() const override
+	{
+		return {
+			{ TEXT("material_name"), TEXT("string"), TEXT("Name for the new material"), true },
+			{ TEXT("path"), TEXT("string"), TEXT("Content path for the material"), false, TEXT("/Game/Materials/") },
+			{ TEXT("base_color_texture"), TEXT("string"), TEXT("Path to texture for base color (optional)"), false },
+			{ TEXT("normal_texture"), TEXT("string"), TEXT("Path to normal map texture (optional)"), false },
+			{ TEXT("roughness_texture"), TEXT("string"), TEXT("Path to roughness texture (optional)"), false },
+			{ TEXT("metallic_texture"), TEXT("string"), TEXT("Path to metallic texture (optional)"), false },
+			{ TEXT("emissive_texture"), TEXT("string"), TEXT("Path to emissive texture (optional)"), false },
+			{ TEXT("base_color"), TEXT("object"), TEXT("Base color as {r, g, b, a} if no texture"), false },
+			{ TEXT("roughness"), TEXT("number"), TEXT("Roughness value (0-1) if no texture"), false, TEXT("0.5") },
+			{ TEXT("metallic"), TEXT("number"), TEXT("Metallic value (0-1) if no texture"), false, TEXT("0.0") },
+			{ TEXT("two_sided"), TEXT("boolean"), TEXT("Enable two-sided rendering"), false, TEXT("false") },
+			{ TEXT("blend_mode"), TEXT("string"), TEXT("Blend mode (Opaque, Masked, Translucent, Additive, Modulate)"), false, TEXT("Opaque") }
+		};
+	}
+	
+	virtual FECACommandResult Execute(const TSharedPtr<FJsonObject>& Params) override;
+};
+
+/**
+ * Create a material from textures in one call (import + material creation)
+ */
+class FECACommand_CreateMaterialFromTextures : public IECACommand
+{
+public:
+	virtual FString GetName() const override { return TEXT("create_material_from_textures"); }
+	virtual FString GetDescription() const override { return TEXT("Import textures and create a material from them in one operation"); }
+	virtual FString GetCategory() const override { return TEXT("Asset"); }
+	
+	virtual TArray<FECACommandParam> GetParameters() const override
+	{
+		return {
+			{ TEXT("material_name"), TEXT("string"), TEXT("Name for the new material"), true },
+			{ TEXT("material_path"), TEXT("string"), TEXT("Content path for the material"), false, TEXT("/Game/Materials/") },
+			{ TEXT("texture_path"), TEXT("string"), TEXT("Content path for imported textures"), false, TEXT("/Game/Textures/") },
+			{ TEXT("base_color_file"), TEXT("string"), TEXT("File path to base color/diffuse image"), false },
+			{ TEXT("normal_file"), TEXT("string"), TEXT("File path to normal map image"), false },
+			{ TEXT("roughness_file"), TEXT("string"), TEXT("File path to roughness image"), false },
+			{ TEXT("metallic_file"), TEXT("string"), TEXT("File path to metallic image"), false },
+			{ TEXT("emissive_file"), TEXT("string"), TEXT("File path to emissive image"), false },
+			{ TEXT("orm_file"), TEXT("string"), TEXT("File path to ORM (Occlusion/Roughness/Metallic) packed texture"), false },
+			{ TEXT("two_sided"), TEXT("boolean"), TEXT("Enable two-sided rendering"), false, TEXT("false") },
+			{ TEXT("blend_mode"), TEXT("string"), TEXT("Blend mode (Opaque, Masked, Translucent, Additive, Modulate)"), false, TEXT("Opaque") }
+		};
+	}
+	
+	virtual FECACommandResult Execute(const TSharedPtr<FJsonObject>& Params) override;
+};
+
+/**
+ * Get texture info (technical metadata only)
+ */
+class FECACommand_GetTextureInfo : public IECACommand
+{
+public:
+	virtual FString GetName() const override { return TEXT("get_texture_info"); }
+	virtual FString GetDescription() const override { return TEXT("Get technical info about a texture (resolution, format, mips). NOTE: For visual description of what a texture looks like, use get_asset_thumbnail + analyze_image instead."); }
+	virtual FString GetCategory() const override { return TEXT("Asset"); }
+	
+	virtual TArray<FECACommandParam> GetParameters() const override
+	{
+		return {
+			{ TEXT("texture_path"), TEXT("string"), TEXT("Path to the texture asset"), true }
+		};
+	}
+	
+	virtual FECACommandResult Execute(const TSharedPtr<FJsonObject>& Params) override;
+};
+
+/**
+ * Get material info (technical/node graph metadata only)
+ */
+class FECACommand_GetMaterialInfo : public IECACommand
+{
+public:
+	virtual FString GetName() const override { return TEXT("get_material_info"); }
+	virtual FString GetDescription() const override { return TEXT("Get technical info about a material (nodes, parameters). NOTE: For visual description of what a material looks like, use get_asset_thumbnail + analyze_image instead."); }
+	virtual FString GetCategory() const override { return TEXT("Asset"); }
+	
+	virtual TArray<FECACommandParam> GetParameters() const override
+	{
+		return {
+			{ TEXT("material_path"), TEXT("string"), TEXT("Path to the material asset"), true }
+		};
+	}
+	
+	virtual FECACommandResult Execute(const TSharedPtr<FJsonObject>& Params) override;
+};
+
+/**
+ * List textures in a path
+ */
+class FECACommand_ListTextures : public IECACommand
+{
+public:
+	virtual FString GetName() const override { return TEXT("list_textures"); }
+	virtual FString GetDescription() const override { return TEXT("List all texture assets in a path. TIP: For visual descriptions, use get_asset_thumbnail + analyze_image to describe what textures look like."); }
+	virtual FString GetCategory() const override { return TEXT("Asset"); }
+	
+	virtual TArray<FECACommandParam> GetParameters() const override
+	{
+		return {
+			{ TEXT("path"), TEXT("string"), TEXT("Content path to search"), false, TEXT("/Game/") },
+			{ TEXT("recursive"), TEXT("boolean"), TEXT("Search recursively"), false, TEXT("true") }
+		};
+	}
+	
+	virtual FECACommandResult Execute(const TSharedPtr<FJsonObject>& Params) override;
+};
+
+/**
+ * List materials in a path
+ */
+class FECACommand_ListMaterials : public IECACommand
+{
+public:
+	virtual FString GetName() const override { return TEXT("list_materials"); }
+	virtual FString GetDescription() const override { return TEXT("List all material assets in a path. TIP: For visual descriptions, use get_asset_thumbnail + analyze_image to describe what materials look like."); }
+	virtual FString GetCategory() const override { return TEXT("Asset"); }
+	
+	virtual TArray<FECACommandParam> GetParameters() const override
+	{
+		return {
+			{ TEXT("path"), TEXT("string"), TEXT("Content path to search"), false, TEXT("/Game/") },
+			{ TEXT("recursive"), TEXT("boolean"), TEXT("Search recursively"), false, TEXT("true") }
+		};
+	}
+	
+	virtual FECACommandResult Execute(const TSharedPtr<FJsonObject>& Params) override;
+};
+
+/**
+ * Set a texture parameter on a material instance
+ */
+class FECACommand_SetMaterialTextureParam : public IECACommand
+{
+public:
+	virtual FString GetName() const override { return TEXT("set_material_texture_param"); }
+	virtual FString GetDescription() const override { return TEXT("Set a texture parameter on a material or material instance"); }
+	virtual FString GetCategory() const override { return TEXT("Asset"); }
+	
+	virtual TArray<FECACommandParam> GetParameters() const override
+	{
+		return {
+			{ TEXT("material_path"), TEXT("string"), TEXT("Path to the material asset"), true },
+			{ TEXT("parameter_name"), TEXT("string"), TEXT("Name of the texture parameter"), true },
+			{ TEXT("texture_path"), TEXT("string"), TEXT("Path to the texture asset"), true }
+		};
+	}
+	
+	virtual FECACommandResult Execute(const TSharedPtr<FJsonObject>& Params) override;
+};
+
+/**
+ * Set material on an actor's mesh component
+ */
+class FECACommand_SetActorMaterial : public IECACommand
+{
+public:
+	virtual FString GetName() const override { return TEXT("set_actor_material"); }
+	virtual FString GetDescription() const override { return TEXT("Set a material on an actor's mesh component (StaticMesh, SkeletalMesh, etc.)"); }
+	virtual FString GetCategory() const override { return TEXT("Asset"); }
+	
+	virtual TArray<FECACommandParam> GetParameters() const override
+	{
+		return {
+			{ TEXT("actor_name"), TEXT("string"), TEXT("Name of the actor in the level"), true },
+			{ TEXT("material_path"), TEXT("string"), TEXT("Path to the material asset"), true },
+			{ TEXT("slot_index"), TEXT("number"), TEXT("Material slot index (0-based)"), false, TEXT("0") },
+			{ TEXT("component_name"), TEXT("string"), TEXT("Specific component name (if actor has multiple mesh components)"), false }
+		};
+	}
+	
+	virtual FECACommandResult Execute(const TSharedPtr<FJsonObject>& Params) override;
+};
+
+/**
+ * Get materials on an actor
+ */
+class FECACommand_GetActorMaterials : public IECACommand
+{
+public:
+	virtual FString GetName() const override { return TEXT("get_actor_materials"); }
+	virtual FString GetDescription() const override { return TEXT("Get all materials currently assigned to an actor's mesh components"); }
+	virtual FString GetCategory() const override { return TEXT("Asset"); }
+	
+	virtual TArray<FECACommandParam> GetParameters() const override
+	{
+		return {
+			{ TEXT("actor_name"), TEXT("string"), TEXT("Name of the actor in the level"), true }
+		};
+	}
+	
+	virtual FECACommandResult Execute(const TSharedPtr<FJsonObject>& Params) override;
+};
+
+/**
+ * Set material on a static mesh asset (default material)
+ */
+class FECACommand_SetStaticMeshMaterial : public IECACommand
+{
+public:
+	virtual FString GetName() const override { return TEXT("set_static_mesh_material"); }
+	virtual FString GetDescription() const override { return TEXT("Set the default material on a static mesh asset"); }
+	virtual FString GetCategory() const override { return TEXT("Asset"); }
+	
+	virtual TArray<FECACommandParam> GetParameters() const override
+	{
+		return {
+			{ TEXT("mesh_path"), TEXT("string"), TEXT("Path to the static mesh asset"), true },
+			{ TEXT("material_path"), TEXT("string"), TEXT("Path to the material asset"), true },
+			{ TEXT("slot_index"), TEXT("number"), TEXT("Material slot index (0-based)"), false, TEXT("0") }
+		};
+	}
+	
+	virtual FECACommandResult Execute(const TSharedPtr<FJsonObject>& Params) override;
+};
+
+/**
+ * Create a material instance from a parent material
+ */
+class FECACommand_CreateMaterialInstance : public IECACommand
+{
+public:
+	virtual FString GetName() const override { return TEXT("create_material_instance"); }
+	virtual FString GetDescription() const override { return TEXT("Create a material instance from a parent material"); }
+	virtual FString GetCategory() const override { return TEXT("Asset"); }
+	
+	virtual TArray<FECACommandParam> GetParameters() const override
+	{
+		return {
+			{ TEXT("instance_name"), TEXT("string"), TEXT("Name for the new material instance"), true },
+			{ TEXT("parent_material"), TEXT("string"), TEXT("Path to the parent material"), true },
+			{ TEXT("path"), TEXT("string"), TEXT("Content path for the material instance"), false, TEXT("/Game/Materials/") }
+		};
+	}
+	
+	virtual FECACommandResult Execute(const TSharedPtr<FJsonObject>& Params) override;
+};
+
+/**
+ * Set scalar parameter on a material instance
+ */
+class FECACommand_SetMaterialInstanceScalarParam : public IECACommand
+{
+public:
+	virtual FString GetName() const override { return TEXT("set_material_instance_scalar_param"); }
+	virtual FString GetDescription() const override { return TEXT("Set a scalar parameter value on a material instance"); }
+	virtual FString GetCategory() const override { return TEXT("Asset"); }
+	
+	virtual TArray<FECACommandParam> GetParameters() const override
+	{
+		return {
+			{ TEXT("instance_path"), TEXT("string"), TEXT("Path to the material instance"), true },
+			{ TEXT("parameter_name"), TEXT("string"), TEXT("Name of the scalar parameter"), true },
+			{ TEXT("value"), TEXT("number"), TEXT("Value to set"), true }
+		};
+	}
+	
+	virtual FECACommandResult Execute(const TSharedPtr<FJsonObject>& Params) override;
+};
+
+/**
+ * Set vector parameter on a material instance
+ */
+class FECACommand_SetMaterialInstanceVectorParam : public IECACommand
+{
+public:
+	virtual FString GetName() const override { return TEXT("set_material_instance_vector_param"); }
+	virtual FString GetDescription() const override { return TEXT("Set a vector/color parameter value on a material instance"); }
+	virtual FString GetCategory() const override { return TEXT("Asset"); }
+	
+	virtual TArray<FECACommandParam> GetParameters() const override
+	{
+		return {
+			{ TEXT("instance_path"), TEXT("string"), TEXT("Path to the material instance"), true },
+			{ TEXT("parameter_name"), TEXT("string"), TEXT("Name of the vector parameter"), true },
+			{ TEXT("value"), TEXT("object"), TEXT("Value as {r, g, b, a} or {x, y, z, w}"), true }
+		};
+	}
+	
+	virtual FECACommandResult Execute(const TSharedPtr<FJsonObject>& Params) override;
+};
+
+/**
+ * Set texture parameter on a material instance
+ */
+class FECACommand_SetMaterialInstanceTextureParam : public IECACommand
+{
+public:
+	virtual FString GetName() const override { return TEXT("set_material_instance_texture_param"); }
+	virtual FString GetDescription() const override { return TEXT("Set a texture parameter on a material instance"); }
+	virtual FString GetCategory() const override { return TEXT("Asset"); }
+	
+	virtual TArray<FECACommandParam> GetParameters() const override
+	{
+		return {
+			{ TEXT("instance_path"), TEXT("string"), TEXT("Path to the material instance"), true },
+			{ TEXT("parameter_name"), TEXT("string"), TEXT("Name of the texture parameter"), true },
+			{ TEXT("texture_path"), TEXT("string"), TEXT("Path to the texture asset"), true }
+		};
+	}
+	
+	virtual FECACommandResult Execute(const TSharedPtr<FJsonObject>& Params) override;
+};
+
+/**
+ * Import an OBJ mesh file as a Static Mesh asset
+ */
+class FECACommand_ImportOBJ : public IECACommand
+{
+public:
+	virtual FString GetName() const override { return TEXT("import_obj"); }
+	virtual FString GetDescription() const override { return TEXT("Import an OBJ mesh file as a Static Mesh asset"); }
+	virtual FString GetCategory() const override { return TEXT("Asset"); }
+	
+	virtual TArray<FECACommandParam> GetParameters() const override
+	{
+		return {
+			{ TEXT("source_path"), TEXT("string"), TEXT("File path to the source OBJ file"), true },
+			{ TEXT("destination_path"), TEXT("string"), TEXT("Content path for the imported mesh (e.g., /Game/Meshes)"), true },
+			{ TEXT("mesh_name"), TEXT("string"), TEXT("Name for the mesh asset (defaults to source filename)"), false },
+			{ TEXT("scale"), TEXT("number"), TEXT("Uniform scale factor to apply (default 1.0, use 100.0 if OBJ is in meters)"), false, TEXT("1.0") },
+			{ TEXT("import_materials"), TEXT("boolean"), TEXT("Import materials from MTL file if present"), false, TEXT("true") },
+			{ TEXT("generate_lightmap_uvs"), TEXT("boolean"), TEXT("Generate lightmap UVs"), false, TEXT("true") },
+			{ TEXT("auto_generate_collision"), TEXT("boolean"), TEXT("Auto-generate simple collision"), false, TEXT("true") },
+			{ TEXT("combine_meshes"), TEXT("boolean"), TEXT("Combine all objects into single mesh"), false, TEXT("true") }
+		};
+	}
+	
+	virtual FECACommandResult Execute(const TSharedPtr<FJsonObject>& Params) override;
+};
+
+/**
+ * Export a texture asset to PNG file
+ */
+class FECACommand_ExportTexture : public IECACommand
+{
+public:
+	virtual FString GetName() const override { return TEXT("export_texture"); }
+	virtual FString GetDescription() const override { return TEXT("Export a texture asset to a PNG file"); }
+	virtual FString GetCategory() const override { return TEXT("Asset"); }
+	
+	virtual TArray<FECACommandParam> GetParameters() const override
+	{
+		return {
+			{ TEXT("texture_path"), TEXT("string"), TEXT("Content path to the texture asset (e.g., /Game/Textures/MyTexture)"), true },
+			{ TEXT("output_path"), TEXT("string"), TEXT("File path for the output PNG file"), true },
+			{ TEXT("mip_level"), TEXT("number"), TEXT("Mip level to export (0 = full resolution)"), false, TEXT("0") },
+		};
+	}
+	
+	virtual FECACommandResult Execute(const TSharedPtr<FJsonObject>& Params) override;
+};
+
+/**
+ * Export a render target to PNG file
+ */
+class FECACommand_ExportRenderTarget : public IECACommand
+{
+public:
+	virtual FString GetName() const override { return TEXT("export_render_target"); }
+	virtual FString GetDescription() const override { return TEXT("Export a render target asset to a PNG file"); }
+	virtual FString GetCategory() const override { return TEXT("Asset"); }
+	
+	virtual TArray<FECACommandParam> GetParameters() const override
+	{
+		return {
+			{ TEXT("render_target_path"), TEXT("string"), TEXT("Content path to the render target asset"), true },
+			{ TEXT("output_path"), TEXT("string"), TEXT("File path for the output PNG file"), true },
+		};
+	}
+	
+	virtual FECACommandResult Execute(const TSharedPtr<FJsonObject>& Params) override;
+};
+
+/**
+ * Set any property on any UObject asset using reflection
+ */
+class FECACommand_SetAssetProperty : public IECACommand
+{
+public:
+	virtual FString GetName() const override { return TEXT("set_asset_property"); }
+	virtual FString GetDescription() const override { return TEXT("Set any UPROPERTY on any asset (UStaticMesh, USkeletalMesh, UTexture, USoundWave, etc.) using reflection. Use get_asset_property to discover available properties."); }
+	virtual FString GetCategory() const override { return TEXT("Asset"); }
+	
+	virtual TArray<FECACommandParam> GetParameters() const override
+	{
+		return {
+			{ TEXT("asset_path"), TEXT("string"), TEXT("Content path to the asset (e.g., /Game/Meshes/MyMesh)"), true },
+			{ TEXT("property"), TEXT("string"), TEXT("Property name (e.g., bAllowCPUAccess, LightMapResolution, NaniteSettings). Supports snake_case."), true },
+			{ TEXT("value"), TEXT("any"), TEXT("Value to set. For enums use the enum value name. For bools use true/false. For structs use nested property paths like 'NaniteSettings.bEnabled'."), true }
+		};
+	}
+	
+	virtual FECACommandResult Execute(const TSharedPtr<FJsonObject>& Params) override;
+};
+
+/**
+ * Get properties from any UObject asset using reflection
+ */
+class FECACommand_GetAssetProperty : public IECACommand
+{
+public:
+	virtual FString GetName() const override { return TEXT("get_asset_property"); }
+	virtual FString GetDescription() const override { return TEXT("Get UPROPERTY values from any asset using reflection. Returns all editable properties if no specific property is requested."); }
+	virtual FString GetCategory() const override { return TEXT("Asset"); }
+	
+	virtual TArray<FECACommandParam> GetParameters() const override
+	{
+		return {
+			{ TEXT("asset_path"), TEXT("string"), TEXT("Content path to the asset (e.g., /Game/Meshes/MyMesh)"), true },
+			{ TEXT("property"), TEXT("string"), TEXT("Specific property name to get (optional - returns all editable properties if not specified)"), false }
+		};
+	}
+	
+	virtual FECACommandResult Execute(const TSharedPtr<FJsonObject>& Params) override;
+};
+
+/**
+ * List all editable properties on an asset class
+ */
+class FECACommand_ListAssetProperties : public IECACommand
+{
+public:
+	virtual FString GetName() const override { return TEXT("list_asset_properties"); }
+	virtual FString GetDescription() const override { return TEXT("List all editable UPROPERTY fields on an asset, organized by category. Useful for discovering what properties can be set."); }
+	virtual FString GetCategory() const override { return TEXT("Asset"); }
+	
+	virtual TArray<FECACommandParam> GetParameters() const override
+	{
+		return {
+			{ TEXT("asset_path"), TEXT("string"), TEXT("Content path to the asset (e.g., /Game/Meshes/MyMesh)"), true },
+			{ TEXT("category"), TEXT("string"), TEXT("Filter by category name (optional)"), false },
+			{ TEXT("search"), TEXT("string"), TEXT("Search filter for property names (optional)"), false }
+		};
+	}
+	
+	virtual FECACommandResult Execute(const TSharedPtr<FJsonObject>& Params) override;
+};
+
+/**
+ * Delete an asset
+ */
+class FECACommand_DeleteAsset : public IECACommand
+{
+public:
+	virtual FString GetName() const override { return TEXT("delete_asset"); }
+	virtual FString GetDescription() const override { return TEXT("Delete an asset from the project. Use with caution!"); }
+	virtual FString GetCategory() const override { return TEXT("Asset"); }
+	
+	virtual TArray<FECACommandParam> GetParameters() const override
+	{
+		return {
+			{ TEXT("asset_path"), TEXT("string"), TEXT("Content path to the asset to delete (e.g., /Game/Meshes/MyMesh)"), true },
+			{ TEXT("force"), TEXT("boolean"), TEXT("Force delete even if asset has references"), false, TEXT("false") }
+		};
+	}
+	
+	virtual FECACommandResult Execute(const TSharedPtr<FJsonObject>& Params) override;
+};
+
+/**
+ * Rename an asset
+ */
+class FECACommand_RenameAsset : public IECACommand
+{
+public:
+	virtual FString GetName() const override { return TEXT("rename_asset"); }
+	virtual FString GetDescription() const override { return TEXT("Rename an asset (changes the asset name, not the path)"); }
+	virtual FString GetCategory() const override { return TEXT("Asset"); }
+	
+	virtual TArray<FECACommandParam> GetParameters() const override
+	{
+		return {
+			{ TEXT("asset_path"), TEXT("string"), TEXT("Content path to the asset to rename (e.g., /Game/Meshes/OldName)"), true },
+			{ TEXT("new_name"), TEXT("string"), TEXT("New name for the asset (just the name, not full path)"), true }
+		};
+	}
+	
+	virtual FECACommandResult Execute(const TSharedPtr<FJsonObject>& Params) override;
+};
+
+/**
+ * Move an asset to a different folder
+ */
+class FECACommand_MoveAsset : public IECACommand
+{
+public:
+	virtual FString GetName() const override { return TEXT("move_asset"); }
+	virtual FString GetDescription() const override { return TEXT("Move an asset to a different folder"); }
+	virtual FString GetCategory() const override { return TEXT("Asset"); }
+	
+	virtual TArray<FECACommandParam> GetParameters() const override
+	{
+		return {
+			{ TEXT("asset_path"), TEXT("string"), TEXT("Content path to the asset to move (e.g., /Game/OldFolder/MyAsset)"), true },
+			{ TEXT("destination_path"), TEXT("string"), TEXT("Destination folder path (e.g., /Game/NewFolder)"), true }
+		};
+	}
+	
+	virtual FECACommandResult Execute(const TSharedPtr<FJsonObject>& Params) override;
+};
+
+/**
+ * Duplicate an asset
+ */
+class FECACommand_DuplicateAsset : public IECACommand
+{
+public:
+	virtual FString GetName() const override { return TEXT("duplicate_asset"); }
+	virtual FString GetDescription() const override { return TEXT("Create a copy of an asset"); }
+	virtual FString GetCategory() const override { return TEXT("Asset"); }
+	
+	virtual TArray<FECACommandParam> GetParameters() const override
+	{
+		return {
+			{ TEXT("asset_path"), TEXT("string"), TEXT("Content path to the asset to duplicate (e.g., /Game/Meshes/MyMesh)"), true },
+			{ TEXT("new_name"), TEXT("string"), TEXT("Name for the duplicated asset"), true },
+			{ TEXT("destination_path"), TEXT("string"), TEXT("Destination folder (defaults to same folder as source)"), false }
+		};
+	}
+	
+	virtual FECACommandResult Execute(const TSharedPtr<FJsonObject>& Params) override;
+};
+
+/**
+ * Import an FBX mesh file as a Static Mesh asset using Interchange
+ */
+class FECACommand_ImportFBX : public IECACommand
+{
+public:
+	virtual FString GetName() const override { return TEXT("import_fbx"); }
+	virtual FString GetDescription() const override { return TEXT("Import an FBX mesh file as a Static Mesh asset using Unreal's Interchange framework"); }
+	virtual FString GetCategory() const override { return TEXT("Asset"); }
+	
+	virtual TArray<FECACommandParam> GetParameters() const override
+	{
+		return {
+			{ TEXT("source_path"), TEXT("string"), TEXT("File path to the source FBX file"), true },
+			{ TEXT("destination_path"), TEXT("string"), TEXT("Content path for the imported mesh (e.g., /Game/Meshes)"), true },
+			{ TEXT("asset_name"), TEXT("string"), TEXT("Name for the imported asset (defaults to source filename)"), false },
+			{ TEXT("import_materials"), TEXT("boolean"), TEXT("Import materials from FBX file"), false, TEXT("true") },
+			{ TEXT("import_textures"), TEXT("boolean"), TEXT("Import textures from FBX file"), false, TEXT("true") },
+			{ TEXT("import_as_skeletal"), TEXT("boolean"), TEXT("Import as Skeletal Mesh instead of Static Mesh (if mesh has bones)"), false, TEXT("false") },
+			{ TEXT("import_animations"), TEXT("boolean"), TEXT("Import animations (only valid for skeletal meshes)"), false, TEXT("false") },
+			{ TEXT("generate_lightmap_uvs"), TEXT("boolean"), TEXT("Generate lightmap UVs for static meshes"), false, TEXT("true") },
+			{ TEXT("auto_generate_collision"), TEXT("boolean"), TEXT("Auto-generate simple collision"), false, TEXT("true") },
+			{ TEXT("combine_meshes"), TEXT("boolean"), TEXT("Combine all meshes into single asset"), false, TEXT("true") },
+			{ TEXT("convert_scene_unit"), TEXT("boolean"), TEXT("Convert scene units to UE units (centimeters)"), false, TEXT("true") },
+			{ TEXT("force_front_x_axis"), TEXT("boolean"), TEXT("Force front axis to X instead of -Y"), false, TEXT("false") }
+		};
+	}
+	
+	virtual FECACommandResult Execute(const TSharedPtr<FJsonObject>& Params) override;
+};
+
+/**
+ * Get the thumbnail/icon of an asset as a PNG image file.
+ * IMPORTANT FOR AI: When generating descriptions of visual assets (meshes, materials, textures, 
+ * Niagara effects, etc.), ALWAYS use this tool to get the thumbnail first, then use analyze_image()
+ * to describe what it looks like visually. This produces much better semantic search descriptions
+ * than just reading metadata. Workflow: 1) get_asset_thumbnail -> 2) analyze_image -> 3) set_asset_association
+ */
+class FECACommand_GetAssetThumbnail : public IECACommand
+{
+public:
+	virtual FString GetName() const override { return TEXT("get_asset_thumbnail"); }
+	virtual FString GetDescription() const override { return TEXT("Get the thumbnail/icon image for any asset and save it to a PNG file. IMPORTANT: For generating visual asset descriptions, use this to get the thumbnail, then analyze_image() to describe it visually, then set_asset_association() to register the description for semantic search."); }
+	virtual FString GetCategory() const override { return TEXT("Asset"); }
+	
+	virtual TArray<FECACommandParam> GetParameters() const override
+	{
+		return {
+			{ TEXT("asset_path"), TEXT("string"), TEXT("Content path to the asset (e.g., /Game/Meshes/MyMesh)"), true },
+			{ TEXT("output_path"), TEXT("string"), TEXT("File path to save the thumbnail PNG (e.g., C:/Thumbnails/mesh.png)"), true },
+			{ TEXT("size"), TEXT("number"), TEXT("Thumbnail size in pixels (default 256, max 1024)"), false, TEXT("256") }
+		};
+	}
+	
+	virtual FECACommandResult Execute(const TSharedPtr<FJsonObject>& Params) override;
+};
+
+/**
+ * Get thumbnails for multiple assets in a single call.
+ * IMPORTANT FOR AI: Use this for bulk visual analysis of assets. For each thumbnail, use analyze_image()
+ * to describe what it looks like, then set_asset_association() to register for semantic search.
+ * Much more efficient than calling get_asset_thumbnail multiple times.
+ */
+class FECACommand_GetAssetThumbnails : public IECACommand
+{
+public:
+	virtual FString GetName() const override { return TEXT("get_asset_thumbnails"); }
+	virtual FString GetDescription() const override { return TEXT("Get thumbnails for multiple assets at once. IMPORTANT: For visual asset cataloging, get thumbnails, then analyze_image() each one to describe visually, then set_asset_association() to register descriptions for semantic search. Much more efficient than single-asset calls."); }
+	virtual FString GetCategory() const override { return TEXT("Asset"); }
+	
+	virtual TArray<FECACommandParam> GetParameters() const override
+	{
+		return {
+			{ TEXT("asset_paths"), TEXT("array"), TEXT("Array of content paths to assets"), true },
+			{ TEXT("output_directory"), TEXT("string"), TEXT("Directory to save thumbnail PNGs (files named after assets, e.g., MyMesh.png)"), true },
+			{ TEXT("size"), TEXT("number"), TEXT("Thumbnail size in pixels (default 256, max 1024)"), false, TEXT("256") }
+		};
+	}
+	
+	virtual FECACommandResult Execute(const TSharedPtr<FJsonObject>& Params) override;
+};
diff --git a/Engine/Plugins/Experimental/ECABridge/Source/ECABridge/Public/Commands/ECABlueprintCommands.h b/Engine/Plugins/Experimental/ECABridge/Source/ECABridge/Public/Commands/ECABlueprintCommands.h
new file mode 100644
--- /dev/null
+++ b/Engine/Plugins/Experimental/ECABridge/Source/ECABridge/Public/Commands/ECABlueprintCommands.h
@@ -0,0 +1,200 @@
+// Copyright Epic Games, Inc. All Rights Reserved.
+
+#pragma once
+
+#include "Commands/ECACommand.h"
+
+/**
+ * Create a new Blueprint asset
+ */
+class FECACommand_CreateBlueprint : public IECACommand
+{
+public:
+	virtual FString GetName() const override { return TEXT("create_blueprint"); }
+	virtual FString GetDescription() const override { return TEXT("Create a new Blueprint asset"); }
+	virtual FString GetCategory() const override { return TEXT("Blueprint"); }
+	
+	virtual TArray<FECACommandParam> GetParameters() const override
+	{
+		return {
+			{ TEXT("blueprint_name"), TEXT("string"), TEXT("Name for the new Blueprint"), true },
+			{ TEXT("parent_class"), TEXT("string"), TEXT("Parent class (Actor, Pawn, Character, or class path)"), false, TEXT("Actor") },
+			{ TEXT("path"), TEXT("string"), TEXT("Content path for the Blueprint"), false, TEXT("/Game/Blueprints/") }
+		};
+	}
+	
+	virtual FECACommandResult Execute(const TSharedPtr<FJsonObject>& Params) override;
+};
+
+/**
+ * Add a component to a Blueprint
+ */
+class FECACommand_AddBlueprintComponent : public IECACommand
+{
+public:
+	virtual FString GetName() const override { return TEXT("add_blueprint_component"); }
+	virtual FString GetDescription() const override { return TEXT("Add a component to a Blueprint"); }
+	virtual FString GetCategory() const override { return TEXT("Blueprint"); }
+	
+	virtual TArray<FECACommandParam> GetParameters() const override
+	{
+		return {
+			{ TEXT("blueprint_path"), TEXT("string"), TEXT("Path to the Blueprint asset"), true },
+			{ TEXT("component_type"), TEXT("string"), TEXT("Type of component (StaticMesh, PointLight, SpotLight, Camera, etc.)"), true },
+			{ TEXT("component_name"), TEXT("string"), TEXT("Name for the component"), false, TEXT("NewComponent") },
+			{ TEXT("attach_to"), TEXT("string"), TEXT("Name of component to attach to"), false }
+		};
+	}
+	
+	virtual FECACommandResult Execute(const TSharedPtr<FJsonObject>& Params) override;
+};
+
+/**
+ * Compile a Blueprint
+ */
+class FECACommand_CompileBlueprint : public IECACommand
+{
+public:
+	virtual FString GetName() const override { return TEXT("compile_blueprint"); }
+	virtual FString GetDescription() const override { return TEXT("Compile a Blueprint"); }
+	virtual FString GetCategory() const override { return TEXT("Blueprint"); }
+	
+	virtual TArray<FECACommandParam> GetParameters() const override
+	{
+		return {
+			{ TEXT("blueprint_path"), TEXT("string"), TEXT("Path to the Blueprint asset"), true }
+		};
+	}
+	
+	virtual FECACommandResult Execute(const TSharedPtr<FJsonObject>& Params) override;
+};
+
+/**
+ * Add a variable to a Blueprint
+ */
+class FECACommand_AddBlueprintVariable : public IECACommand
+{
+public:
+	virtual FString GetName() const override { return TEXT("add_blueprint_variable"); }
+	virtual FString GetDescription() const override { return TEXT("Add a member variable to a Blueprint"); }
+	virtual FString GetCategory() const override { return TEXT("Blueprint"); }
+	
+	virtual TArray<FECACommandParam> GetParameters() const override
+	{
+		return {
+			{ TEXT("blueprint_path"), TEXT("string"), TEXT("Path to the Blueprint asset"), true },
+			{ TEXT("variable_name"), TEXT("string"), TEXT("Name for the variable"), true },
+			{ TEXT("variable_type"), TEXT("string"), TEXT("Type of variable (Boolean, Integer, Float, String, Vector, Rotator, Transform)"), true },
+			{ TEXT("default_value"), TEXT("any"), TEXT("Default value for the variable"), false },
+			{ TEXT("is_instance_editable"), TEXT("boolean"), TEXT("Allow editing per instance"), false, TEXT("false") },
+			{ TEXT("is_blueprint_read_only"), TEXT("boolean"), TEXT("Make read-only in Blueprints"), false, TEXT("false") }
+		};
+	}
+	
+	virtual FECACommandResult Execute(const TSharedPtr<FJsonObject>& Params) override;
+};
+
+/**
+ * Spawn a Blueprint actor in the level
+ */
+class FECACommand_SpawnBlueprintActor : public IECACommand
+{
+public:
+	virtual FString GetName() const override { return TEXT("spawn_blueprint_actor"); }
+	virtual FString GetDescription() const override { return TEXT("Spawn an instance of a Blueprint in the level"); }
+	virtual FString GetCategory() const override { return TEXT("Blueprint"); }
+	
+	virtual TArray<FECACommandParam> GetParameters() const override
+	{
+		return {
+			{ TEXT("blueprint_path"), TEXT("string"), TEXT("Path to the Blueprint asset"), true },
+			{ TEXT("location"), TEXT("object"), TEXT("Spawn location {x, y, z}"), false },
+			{ TEXT("rotation"), TEXT("object"), TEXT("Spawn rotation {pitch, yaw, roll}"), false },
+			{ TEXT("actor_name"), TEXT("string"), TEXT("Display name for the spawned actor"), false }
+		};
+	}
+	
+	virtual FECACommandResult Execute(const TSharedPtr<FJsonObject>& Params) override;
+};
+
+/**
+ * Get Blueprint info
+ */
+class FECACommand_GetBlueprintInfo : public IECACommand
+{
+public:
+	virtual FString GetName() const override { return TEXT("get_blueprint_info"); }
+	virtual FString GetDescription() const override { return TEXT("Get information about a Blueprint"); }
+	virtual FString GetCategory() const override { return TEXT("Blueprint"); }
+	
+	virtual TArray<FECACommandParam> GetParameters() const override
+	{
+		return {
+			{ TEXT("blueprint_path"), TEXT("string"), TEXT("Path to the Blueprint asset"), true }
+		};
+	}
+	
+	virtual FECACommandResult Execute(const TSharedPtr<FJsonObject>& Params) override;
+};
+
+/**
+ * List all Blueprints in a path
+ */
+class FECACommand_ListBlueprints : public IECACommand
+{
+public:
+	virtual FString GetName() const override { return TEXT("list_blueprints"); }
+	virtual FString GetDescription() const override { return TEXT("List all Blueprint assets in a path"); }
+	virtual FString GetCategory() const override { return TEXT("Blueprint"); }
+	
+	virtual TArray<FECACommandParam> GetParameters() const override
+	{
+		return {
+			{ TEXT("path"), TEXT("string"), TEXT("Content path to search"), false, TEXT("/Game/") },
+			{ TEXT("recursive"), TEXT("boolean"), TEXT("Search recursively"), false, TEXT("true") }
+		};
+	}
+	
+	virtual FECACommandResult Execute(const TSharedPtr<FJsonObject>& Params) override;
+};
+
+/**
+ * Open Blueprint editor
+ */
+class FECACommand_OpenBlueprintEditor : public IECACommand
+{
+public:
+	virtual FString GetName() const override { return TEXT("open_blueprint_editor"); }
+	virtual FString GetDescription() const override { return TEXT("Open a Blueprint in the Blueprint editor"); }
+	virtual FString GetCategory() const override { return TEXT("Blueprint"); }
+	
+	virtual TArray<FECACommandParam> GetParameters() const override
+	{
+		return {
+			{ TEXT("blueprint_path"), TEXT("string"), TEXT("Path to the Blueprint asset"), true }
+		};
+	}
+	
+	virtual FECACommandResult Execute(const TSharedPtr<FJsonObject>& Params) override;
+};
+
+/**
+ * Delete a function from a Blueprint
+ */
+class FECACommand_DeleteBlueprintFunction : public IECACommand
+{
+public:
+	virtual FString GetName() const override { return TEXT("delete_blueprint_function"); }
+	virtual FString GetDescription() const override { return TEXT("Delete a function graph from a Blueprint"); }
+	virtual FString GetCategory() const override { return TEXT("Blueprint"); }
+	
+	virtual TArray<FECACommandParam> GetParameters() const override
+	{
+		return {
+			{ TEXT("blueprint_path"), TEXT("string"), TEXT("Path to the Blueprint asset"), true },
+			{ TEXT("function_name"), TEXT("string"), TEXT("Name of the function to delete"), true }
+		};
+	}
+	
+	virtual FECACommandResult Execute(const TSharedPtr<FJsonObject>& Params) override;
+};
diff --git a/Engine/Plugins/Experimental/ECABridge/Source/ECABridge/Public/Commands/ECABlueprintLispCommands.h b/Engine/Plugins/Experimental/ECABridge/Source/ECABridge/Public/Commands/ECABlueprintLispCommands.h
new file mode 100644
--- /dev/null
+++ b/Engine/Plugins/Experimental/ECABridge/Source/ECABridge/Public/Commands/ECABlueprintLispCommands.h
@@ -0,0 +1,114 @@
+// Copyright Epic Games, Inc. All Rights Reserved.
+
+#pragma once
+
+#include "Commands/ECACommand.h"
+
+/**
+ * Parse and validate BlueprintLisp code
+ * 
+ * BlueprintLisp is a LISP-like DSL for representing Blueprint graphs in a format
+ * that's easy for AI to read, understand, and generate.
+ * 
+ * PREFERRED: Use BlueprintLisp format for all Blueprint implementation tasks.
+ * It is more concise, easier for AI to reason about, and less error-prone than
+ * the JSON-based node commands.
+ * 
+ * Example:
+ *   (event BeginPlay
+ *     (let player (GetPlayerCharacter 0))
+ *     (branch (IsValid player)
+ *       :true (PrintString "Valid!")
+ *       :false (PrintString "Invalid!")))
+ */
+class FECACommand_ParseBlueprintLisp : public IECACommand
+{
+public:
+	virtual FString GetName() const override { return TEXT("parse_blueprint_lisp"); }
+	virtual FString GetDescription() const override { return TEXT("Parse and validate BlueprintLisp code. Returns the AST or parse errors."); }
+	virtual FString GetCategory() const override { return TEXT("BlueprintLisp"); }
+	
+	virtual TArray<FECACommandParam> GetParameters() const override
+	{
+		return {
+			{ TEXT("code"), TEXT("string"), TEXT("BlueprintLisp source code to parse"), true },
+			{ TEXT("pretty_print"), TEXT("boolean"), TEXT("Return pretty-printed version of the code"), false }
+		};
+	}
+	
+	virtual FECACommandResult Execute(const TSharedPtr<FJsonObject>& Params) override;
+};
+
+/**
+ * Convert a Blueprint graph to BlueprintLisp code
+ * 
+ * PREFERRED: Use this to read Blueprint logic. The Lisp format is more compact
+ * and easier to understand than raw node JSON.
+ */
+class FECACommand_BlueprintToLisp : public IECACommand
+{
+public:
+	virtual FString GetName() const override { return TEXT("blueprint_to_lisp"); }
+	virtual FString GetDescription() const override { return TEXT("PREFERRED: Convert a Blueprint graph to BlueprintLisp code for easy reading and AI manipulation."); }
+	virtual FString GetCategory() const override { return TEXT("BlueprintLisp"); }
+	
+	virtual TArray<FECACommandParam> GetParameters() const override
+	{
+		return {
+			{ TEXT("blueprint_path"), TEXT("string"), TEXT("Path to the Blueprint asset"), true },
+			{ TEXT("graph_name"), TEXT("string"), TEXT("Name of the graph to convert (default: EventGraph)"), false },
+			{ TEXT("include_comments"), TEXT("boolean"), TEXT("Include node comments in output"), false },
+			{ TEXT("include_positions"), TEXT("boolean"), TEXT("Include node positions as metadata"), false }
+		};
+	}
+	
+	virtual FECACommandResult Execute(const TSharedPtr<FJsonObject>& Params) override;
+};
+
+/**
+ * Convert BlueprintLisp code to Blueprint graph nodes
+ * 
+ * PREFERRED: Use this to implement Blueprint logic. The Lisp format handles
+ * node creation, wiring, and positioning automatically. Much simpler than
+ * manually creating nodes and connections via JSON commands.
+ */
+class FECACommand_LispToBlueprint : public IECACommand
+{
+public:
+	virtual FString GetName() const override { return TEXT("lisp_to_blueprint"); }
+	virtual FString GetDescription() const override { return TEXT("PREFERRED: Convert BlueprintLisp code to Blueprint graph nodes. Handles node creation, wiring, and layout automatically."); }
+	virtual FString GetCategory() const override { return TEXT("BlueprintLisp"); }
+	
+	virtual TArray<FECACommandParam> GetParameters() const override
+	{
+		return {
+			{ TEXT("blueprint_path"), TEXT("string"), TEXT("Path to the Blueprint asset"), true },
+			{ TEXT("code"), TEXT("string"), TEXT("BlueprintLisp source code to convert"), true },
+			{ TEXT("graph_name"), TEXT("string"), TEXT("Name of the graph to modify (default: EventGraph)"), false },
+			{ TEXT("clear_existing"), TEXT("boolean"), TEXT("Clear existing nodes before adding new ones"), false },
+			{ TEXT("auto_layout"), TEXT("boolean"), TEXT("Automatically arrange nodes after creation (default: true)"), false }
+		};
+	}
+	
+	virtual FECACommandResult Execute(const TSharedPtr<FJsonObject>& Params) override;
+};
+
+/**
+ * Get BlueprintLisp syntax help and examples
+ */
+class FECACommand_BlueprintLispHelp : public IECACommand
+{
+public:
+	virtual FString GetName() const override { return TEXT("blueprint_lisp_help"); }
+	virtual FString GetDescription() const override { return TEXT("Get BlueprintLisp syntax documentation and examples."); }
+	virtual FString GetCategory() const override { return TEXT("BlueprintLisp"); }
+	
+	virtual TArray<FECACommandParam> GetParameters() const override
+	{
+		return {
+			{ TEXT("topic"), TEXT("string"), TEXT("Specific topic: forms, events, functions, flow, expressions, examples (default: all)"), false }
+		};
+	}
+	
+	virtual FECACommandResult Execute(const TSharedPtr<FJsonObject>& Params) override;
+};
diff --git a/Engine/Plugins/Experimental/ECABridge/Source/ECABridge/Public/Commands/ECABlueprintNodeCommands.h b/Engine/Plugins/Experimental/ECABridge/Source/ECABridge/Public/Commands/ECABlueprintNodeCommands.h
new file mode 100644
--- /dev/null
+++ b/Engine/Plugins/Experimental/ECABridge/Source/ECABridge/Public/Commands/ECABlueprintNodeCommands.h
@@ -0,0 +1,613 @@
+// Copyright Epic Games, Inc. All Rights Reserved.
+
+#pragma once
+
+#include "Commands/ECACommand.h"
+
+/**
+ * Add an event node to a Blueprint's event graph
+ */
+class FECACommand_AddBlueprintEventNode : public IECACommand
+{
+public:
+	virtual FString GetName() const override { return TEXT("add_blueprint_event_node"); }
+	virtual FString GetDescription() const override { return TEXT("Add an event node (BeginPlay, Tick, etc.) to a Blueprint's event graph"); }
+	virtual FString GetCategory() const override { return TEXT("Blueprint Node"); }
+	
+	virtual TArray<FECACommandParam> GetParameters() const override
+	{
+		return {
+			{ TEXT("blueprint_path"), TEXT("string"), TEXT("Path to the Blueprint asset"), true },
+			{ TEXT("event_name"), TEXT("string"), TEXT("Event name: ReceiveBeginPlay, ReceiveTick, ReceiveActorBeginOverlap, etc."), true },
+			{ TEXT("node_position"), TEXT("object"), TEXT("Node position in graph {x, y}"), false }
+		};
+	}
+	
+	virtual FECACommandResult Execute(const TSharedPtr<FJsonObject>& Params) override;
+};
+
+/**
+ * Add an input action event node to a Blueprint
+ */
+class FECACommand_AddBlueprintInputActionNode : public IECACommand
+{
+public:
+	virtual FString GetName() const override { return TEXT("add_blueprint_input_action_node"); }
+	virtual FString GetDescription() const override { return TEXT("Add an input action event node to a Blueprint"); }
+	virtual FString GetCategory() const override { return TEXT("Blueprint Node"); }
+	
+	virtual TArray<FECACommandParam> GetParameters() const override
+	{
+		return {
+			{ TEXT("blueprint_path"), TEXT("string"), TEXT("Path to the Blueprint asset"), true },
+			{ TEXT("action_name"), TEXT("string"), TEXT("Name of the input action"), true },
+			{ TEXT("node_position"), TEXT("object"), TEXT("Node position in graph {x, y}"), false }
+		};
+	}
+	
+	virtual FECACommandResult Execute(const TSharedPtr<FJsonObject>& Params) override;
+};
+
+/**
+ * Add a function call node to a Blueprint
+ */
+class FECACommand_AddBlueprintFunctionNode : public IECACommand
+{
+public:
+	virtual FString GetName() const override { return TEXT("add_blueprint_function_node"); }
+	virtual FString GetDescription() const override { return TEXT("Add a function call node to a Blueprint graph"); }
+	virtual FString GetCategory() const override { return TEXT("Blueprint Node"); }
+	
+	virtual TArray<FECACommandParam> GetParameters() const override
+	{
+		return {
+			{ TEXT("blueprint_path"), TEXT("string"), TEXT("Path to the Blueprint asset"), true },
+			{ TEXT("function_name"), TEXT("string"), TEXT("Name of the function to call (e.g., PrintString, SetActorLocation)"), true },
+			{ TEXT("target"), TEXT("string"), TEXT("Target for the function call (self, component name, or class)"), false, TEXT("self") },
+			{ TEXT("target_class"), TEXT("string"), TEXT("Class containing the function (e.g., KismetSystemLibrary, Actor)"), false },
+			{ TEXT("graph_name"), TEXT("string"), TEXT("Name of the graph to add to (EventGraph or function name)"), false, TEXT("EventGraph") },
+			{ TEXT("node_position"), TEXT("object"), TEXT("Node position in graph {x, y}"), false }
+		};
+	}
+	
+	virtual FECACommandResult Execute(const TSharedPtr<FJsonObject>& Params) override;
+};
+
+/**
+ * Add a macro node (ForLoop, WhileLoop, ForEachLoop, etc.) to a Blueprint
+ * 
+ * Common macros: ForLoop, ForLoopWithBreak, WhileLoop, ForEachLoop, ForEachLoopWithBreak,
+ * Gate, MultiGate, DoOnce, DoN, FlipFlop, Sequence, IsValid
+ * 
+ * Example: {"blueprint_path":"/Game/BP_Test", "macro_name":"ForLoop", "node_position":{"x":200,"y":100}}
+ */
+class FECACommand_AddBlueprintMacroNode : public IECACommand
+{
+public:
+	virtual FString GetName() const override { return TEXT("add_blueprint_macro_node"); }
+	virtual FString GetDescription() const override { return TEXT("Add a macro node to a Blueprint. Common macros: ForLoop, ForLoopWithBreak, WhileLoop, ForEachLoop, Gate, MultiGate, DoOnce, DoN, FlipFlop, Sequence, IsValid. Example: {\"macro_name\":\"ForLoop\"}"); }
+	virtual FString GetCategory() const override { return TEXT("Blueprint Node"); }
+	
+	virtual TArray<FECACommandParam> GetParameters() const override
+	{
+		return {
+			{ TEXT("blueprint_path"), TEXT("string"), TEXT("Path to the Blueprint asset"), true },
+			{ TEXT("macro_name"), TEXT("string"), TEXT("Name of the macro: ForLoop, WhileLoop, ForEachLoop, Gate, etc."), true },
+			{ TEXT("graph_name"), TEXT("string"), TEXT("Name of the graph to add to"), false, TEXT("EventGraph") },
+			{ TEXT("node_position"), TEXT("object"), TEXT("Node position in graph {x, y}"), false }
+		};
+	}
+	
+	virtual FECACommandResult Execute(const TSharedPtr<FJsonObject>& Params) override;
+};
+
+/**
+ * Add a cast node (Cast To X) to a Blueprint
+ * 
+ * Example: {"blueprint_path":"/Game/BP_Test", "target_class":"Character", "node_position":{"x":200,"y":100}}
+ * For Blueprint classes: {"target_class":"/Game/Blueprints/BP_Enemy"}
+ */
+class FECACommand_AddBlueprintCastNode : public IECACommand
+{
+public:
+	virtual FString GetName() const override { return TEXT("add_blueprint_cast_node"); }
+	virtual FString GetDescription() const override { return TEXT("Add a Cast To node. For native classes use name like Character, Actor, Pawn. For Blueprints use full path like /Game/BP_Enemy. Example: {\"target_class\":\"Character\"}"); }
+	virtual FString GetCategory() const override { return TEXT("Blueprint Node"); }
+	
+	virtual TArray<FECACommandParam> GetParameters() const override
+	{
+		return {
+			{ TEXT("blueprint_path"), TEXT("string"), TEXT("Path to the Blueprint asset"), true },
+			{ TEXT("target_class"), TEXT("string"), TEXT("Class to cast to (e.g., Character, Pawn, or /Game/BP_Enemy for Blueprints)"), true },
+			{ TEXT("pure"), TEXT("boolean"), TEXT("Create a pure cast (no exec pins)"), false, TEXT("false") },
+			{ TEXT("graph_name"), TEXT("string"), TEXT("Name of the graph to add to"), false, TEXT("EventGraph") },
+			{ TEXT("node_position"), TEXT("object"), TEXT("Node position in graph {x, y}"), false }
+		};
+	}
+	
+	virtual FECACommandResult Execute(const TSharedPtr<FJsonObject>& Params) override;
+};
+
+/**
+ * Connect two nodes in a Blueprint graph
+ */
+class FECACommand_ConnectBlueprintNodes : public IECACommand
+{
+public:
+	virtual FString GetName() const override { return TEXT("connect_blueprint_nodes"); }
+	virtual FString GetDescription() const override { return TEXT("Connect two nodes in a Blueprint graph by their pins"); }
+	virtual FString GetCategory() const override { return TEXT("Blueprint Node"); }
+	
+	virtual TArray<FECACommandParam> GetParameters() const override
+	{
+		return {
+			{ TEXT("blueprint_path"), TEXT("string"), TEXT("Path to the Blueprint asset"), true },
+			{ TEXT("source_node_id"), TEXT("string"), TEXT("GUID of the source node"), true },
+			{ TEXT("source_pin"), TEXT("string"), TEXT("Name of the output pin on the source node"), true },
+			{ TEXT("target_node_id"), TEXT("string"), TEXT("GUID of the target node"), true },
+			{ TEXT("target_pin"), TEXT("string"), TEXT("Name of the input pin on the target node"), true },
+			{ TEXT("graph_name"), TEXT("string"), TEXT("Name of the graph"), false, TEXT("EventGraph") }
+		};
+	}
+	
+	virtual FECACommandResult Execute(const TSharedPtr<FJsonObject>& Params) override;
+};
+
+/**
+ * Add a self reference node to a Blueprint
+ */
+class FECACommand_AddBlueprintSelfReference : public IECACommand
+{
+public:
+	virtual FString GetName() const override { return TEXT("add_blueprint_self_reference"); }
+	virtual FString GetDescription() const override { return TEXT("Add a 'Get Self' reference node to a Blueprint graph"); }
+	virtual FString GetCategory() const override { return TEXT("Blueprint Node"); }
+	
+	virtual TArray<FECACommandParam> GetParameters() const override
+	{
+		return {
+			{ TEXT("blueprint_path"), TEXT("string"), TEXT("Path to the Blueprint asset"), true },
+			{ TEXT("graph_name"), TEXT("string"), TEXT("Name of the graph"), false, TEXT("EventGraph") },
+			{ TEXT("node_position"), TEXT("object"), TEXT("Node position in graph {x, y}"), false }
+		};
+	}
+	
+	virtual FECACommandResult Execute(const TSharedPtr<FJsonObject>& Params) override;
+};
+
+/**
+ * Add a component reference (getter) node to a Blueprint
+ */
+class FECACommand_AddBlueprintComponentReference : public IECACommand
+{
+public:
+	virtual FString GetName() const override { return TEXT("add_blueprint_component_reference"); }
+	virtual FString GetDescription() const override { return TEXT("Add a component reference (getter) node to a Blueprint graph"); }
+	virtual FString GetCategory() const override { return TEXT("Blueprint Node"); }
+	
+	virtual TArray<FECACommandParam> GetParameters() const override
+	{
+		return {
+			{ TEXT("blueprint_path"), TEXT("string"), TEXT("Path to the Blueprint asset"), true },
+			{ TEXT("component_name"), TEXT("string"), TEXT("Name of the component to reference"), true },
+			{ TEXT("graph_name"), TEXT("string"), TEXT("Name of the graph"), false, TEXT("EventGraph") },
+			{ TEXT("node_position"), TEXT("object"), TEXT("Node position in graph {x, y}"), false }
+		};
+	}
+	
+	virtual FECACommandResult Execute(const TSharedPtr<FJsonObject>& Params) override;
+};
+
+/**
+ * Find nodes in a Blueprint graph
+ */
+class FECACommand_FindBlueprintNodes : public IECACommand
+{
+public:
+	virtual FString GetName() const override { return TEXT("find_blueprint_nodes"); }
+	virtual FString GetDescription() const override { return TEXT("Find nodes in a Blueprint graph by type or name"); }
+	virtual FString GetCategory() const override { return TEXT("Blueprint Node"); }
+	
+	virtual TArray<FECACommandParam> GetParameters() const override
+	{
+		return {
+			{ TEXT("blueprint_path"), TEXT("string"), TEXT("Path to the Blueprint asset"), true },
+			{ TEXT("graph_name"), TEXT("string"), TEXT("Name of the graph to search"), false, TEXT("EventGraph") },
+			{ TEXT("node_class"), TEXT("string"), TEXT("Filter by node class name"), false },
+			{ TEXT("node_title"), TEXT("string"), TEXT("Filter by node title (partial match)"), false }
+		};
+	}
+	
+	virtual FECACommandResult Execute(const TSharedPtr<FJsonObject>& Params) override;
+};
+
+/**
+ * Add a variable get node to a Blueprint
+ */
+class FECACommand_AddBlueprintVariableGetNode : public IECACommand
+{
+public:
+	virtual FString GetName() const override { return TEXT("add_blueprint_variable_get"); }
+	virtual FString GetDescription() const override { return TEXT("Add a variable getter node to a Blueprint graph"); }
+	virtual FString GetCategory() const override { return TEXT("Blueprint Node"); }
+	
+	virtual TArray<FECACommandParam> GetParameters() const override
+	{
+		return {
+			{ TEXT("blueprint_path"), TEXT("string"), TEXT("Path to the Blueprint asset"), true },
+			{ TEXT("variable_name"), TEXT("string"), TEXT("Name of the variable to get"), true },
+			{ TEXT("graph_name"), TEXT("string"), TEXT("Name of the graph"), false, TEXT("EventGraph") },
+			{ TEXT("node_position"), TEXT("object"), TEXT("Node position in graph {x, y}"), false }
+		};
+	}
+	
+	virtual FECACommandResult Execute(const TSharedPtr<FJsonObject>& Params) override;
+};
+
+/**
+ * Add a variable set node to a Blueprint
+ */
+class FECACommand_AddBlueprintVariableSetNode : public IECACommand
+{
+public:
+	virtual FString GetName() const override { return TEXT("add_blueprint_variable_set"); }
+	virtual FString GetDescription() const override { return TEXT("Add a variable setter node to a Blueprint graph"); }
+	virtual FString GetCategory() const override { return TEXT("Blueprint Node"); }
+	
+	virtual TArray<FECACommandParam> GetParameters() const override
+	{
+		return {
+			{ TEXT("blueprint_path"), TEXT("string"), TEXT("Path to the Blueprint asset"), true },
+			{ TEXT("variable_name"), TEXT("string"), TEXT("Name of the variable to set"), true },
+			{ TEXT("graph_name"), TEXT("string"), TEXT("Name of the graph"), false, TEXT("EventGraph") },
+			{ TEXT("node_position"), TEXT("object"), TEXT("Node position in graph {x, y}"), false }
+		};
+	}
+	
+	virtual FECACommandResult Execute(const TSharedPtr<FJsonObject>& Params) override;
+};
+
+/**
+ * Batch create nodes and connections in a single operation.
+ * 
+ * This command allows you to create multiple nodes and wire them together atomically.
+ * Nodes are assigned temporary IDs (strings you define) which can be used to reference
+ * them in the connections array. The response includes a mapping of your temporary IDs
+ * to the actual node GUIDs.
+ * 
+ * Example usage:
+ * {
+ *   "blueprint_path": "/Game/MyBlueprint",
+ *   "graph_name": "EventGraph",
+ *   "nodes": [
+ *     { "temp_id": "event", "type": "event", "event_name": "ReceiveBeginPlay" },
+ *     { "temp_id": "print", "type": "function", "function_name": "PrintString", "node_position": {"x": 300, "y": 0} }
+ *   ],
+ *   "connections": [
+ *     { "source_node": "event", "source_pin": "then", "target_node": "print", "target_pin": "execute" }
+ *   ]
+ * }
+ */
+class FECACommand_BatchEditBlueprintNodes : public IECACommand
+{
+public:
+	virtual FString GetName() const override { return TEXT("batch_edit_blueprint_nodes"); }
+	virtual FString GetDescription() const override { return TEXT("Create multiple nodes and wire them together atomically. Node types: event, custom_event, function, variable_get, variable_set, self, component, input_action, cast, macro. Cast example: {\"type\":\"cast\", \"target_class\":\"Character\"}. Macro example: {\"type\":\"macro\", \"macro_name\":\"ForLoop\"}"); }
+	virtual FString GetCategory() const override { return TEXT("Blueprint Node"); }
+	
+	virtual TArray<FECACommandParam> GetParameters() const override
+	{
+		return {
+			{ TEXT("blueprint_path"), TEXT("string"), TEXT("Path to the Blueprint asset"), true },
+			{ TEXT("graph_name"), TEXT("string"), TEXT("Name of the graph (EventGraph or function name)"), false, TEXT("EventGraph") },
+			{ TEXT("nodes"), TEXT("array"), TEXT("Array of node definitions with temp_id and type. Types: event, custom_event, function, variable_get, variable_set, self, component, input_action, cast (target_class), macro (macro_name)"), true },
+			{ TEXT("connections"), TEXT("array"), TEXT("Array of connection definitions: {source_node, source_pin, target_node, target_pin}"), false }
+		};
+	}
+	
+	virtual FECACommandResult Execute(const TSharedPtr<FJsonObject>& Params) override;
+};
+
+/**
+ * Get detailed information about a node's pins.
+ * Useful for discovering pin names before making connections.
+ */
+class FECACommand_GetBlueprintNodePins : public IECACommand
+{
+public:
+	virtual FString GetName() const override { return TEXT("get_blueprint_node_pins"); }
+	virtual FString GetDescription() const override { return TEXT("Get detailed pin information for a node by its GUID. Use this to discover pin names for connections."); }
+	virtual FString GetCategory() const override { return TEXT("Blueprint Node"); }
+	
+	virtual TArray<FECACommandParam> GetParameters() const override
+	{
+		return {
+			{ TEXT("blueprint_path"), TEXT("string"), TEXT("Path to the Blueprint asset"), true },
+			{ TEXT("node_id"), TEXT("string"), TEXT("GUID of the node"), true },
+			{ TEXT("graph_name"), TEXT("string"), TEXT("Name of the graph"), false, TEXT("EventGraph") }
+		};
+	}
+	
+	virtual FECACommandResult Execute(const TSharedPtr<FJsonObject>& Params) override;
+};
+
+
+/**
+ * Delete a node from a Blueprint graph by its ID
+ */
+class FECACommand_DeleteBlueprintNode : public IECACommand
+{
+public:
+	virtual FString GetName() const override { return TEXT("delete_blueprint_node"); }
+	virtual FString GetDescription() const override { return TEXT("Delete a node from a Blueprint graph by its GUID"); }
+	virtual FString GetCategory() const override { return TEXT("Blueprint Node"); }
+	
+	virtual TArray<FECACommandParam> GetParameters() const override
+	{
+		return {
+			{ TEXT("blueprint_path"), TEXT("string"), TEXT("Path to the Blueprint asset"), true },
+			{ TEXT("node_id"), TEXT("string"), TEXT("GUID of the node to delete"), true },
+			{ TEXT("graph_name"), TEXT("string"), TEXT("Name of the graph"), false, TEXT("EventGraph") }
+		};
+	}
+	
+	virtual FECACommandResult Execute(const TSharedPtr<FJsonObject>& Params) override;
+};
+
+/**
+ * Find and delete orphan nodes (nodes with no exec pin connections) from a Blueprint graph
+ */
+class FECACommand_CleanupOrphanNodes : public IECACommand
+{
+public:
+	virtual FString GetName() const override { return TEXT("cleanup_orphan_nodes"); }
+	virtual FString GetDescription() const override { return TEXT("Find and optionally delete orphan nodes (nodes with no execution flow connections) from a Blueprint graph"); }
+	virtual FString GetCategory() const override { return TEXT("Blueprint Node"); }
+	
+	virtual TArray<FECACommandParam> GetParameters() const override
+	{
+		return {
+			{ TEXT("blueprint_path"), TEXT("string"), TEXT("Path to the Blueprint asset"), true },
+			{ TEXT("graph_name"), TEXT("string"), TEXT("Name of the graph"), false, TEXT("EventGraph") },
+			{ TEXT("delete"), TEXT("boolean"), TEXT("If true, delete the orphan nodes. If false, just report them."), false, TEXT("false") },
+			{ TEXT("node_class_filter"), TEXT("string"), TEXT("Only consider nodes of this class (e.g., K2Node_CallFunction)"), false }
+		};
+	}
+	
+	virtual FECACommandResult Execute(const TSharedPtr<FJsonObject>& Params) override;
+};
+
+/**
+ * Disconnect all pins on a node
+ */
+class FECACommand_DisconnectBlueprintNode : public IECACommand
+{
+public:
+	virtual FString GetName() const override { return TEXT("disconnect_blueprint_node"); }
+	virtual FString GetDescription() const override { return TEXT("Disconnect all pins or specific pins on a Blueprint node"); }
+	virtual FString GetCategory() const override { return TEXT("Blueprint Node"); }
+	
+	virtual TArray<FECACommandParam> GetParameters() const override
+	{
+		return {
+			{ TEXT("blueprint_path"), TEXT("string"), TEXT("Path to the Blueprint asset"), true },
+			{ TEXT("node_id"), TEXT("string"), TEXT("GUID of the node"), true },
+			{ TEXT("pin_name"), TEXT("string"), TEXT("Specific pin to disconnect (optional, disconnects all if not specified)"), false },
+			{ TEXT("graph_name"), TEXT("string"), TEXT("Name of the graph"), false, TEXT("EventGraph") }
+		};
+	}
+	
+	virtual FECACommandResult Execute(const TSharedPtr<FJsonObject>& Params) override;
+};
+
+
+/**
+ * Set the default value of a pin on a Blueprint node
+ */
+class FECACommand_SetBlueprintPinValue : public IECACommand
+{
+public:
+	virtual FString GetName() const override { return TEXT("set_blueprint_pin_value"); }
+	virtual FString GetDescription() const override { return TEXT("Set the default value of an input pin on a Blueprint node"); }
+	virtual FString GetCategory() const override { return TEXT("Blueprint Node"); }
+	
+	virtual TArray<FECACommandParam> GetParameters() const override
+	{
+		return {
+			{ TEXT("blueprint_path"), TEXT("string"), TEXT("Path to the Blueprint asset"), true },
+			{ TEXT("node_id"), TEXT("string"), TEXT("GUID of the node"), true },
+			{ TEXT("pin_name"), TEXT("string"), TEXT("Name of the pin to set"), true },
+			{ TEXT("value"), TEXT("string"), TEXT("Value to set (as string - will be parsed based on pin type)"), true },
+			{ TEXT("graph_name"), TEXT("string"), TEXT("Name of the graph"), false, TEXT("EventGraph") }
+		};
+	}
+	
+	virtual FECACommandResult Execute(const TSharedPtr<FJsonObject>& Params) override;
+};
+
+/**
+ * Set the default value of a Blueprint variable
+ */
+class FECACommand_SetBlueprintVariableDefault : public IECACommand
+{
+public:
+	virtual FString GetName() const override { return TEXT("set_blueprint_variable_default"); }
+	virtual FString GetDescription() const override { return TEXT("Set the default value of a Blueprint variable"); }
+	virtual FString GetCategory() const override { return TEXT("Blueprint"); }
+	
+	virtual TArray<FECACommandParam> GetParameters() const override
+	{
+		return {
+			{ TEXT("blueprint_path"), TEXT("string"), TEXT("Path to the Blueprint asset"), true },
+			{ TEXT("variable_name"), TEXT("string"), TEXT("Name of the variable"), true },
+			{ TEXT("value"), TEXT("any"), TEXT("Default value to set (type must match variable type)"), true }
+		};
+	}
+	
+	virtual FECACommandResult Execute(const TSharedPtr<FJsonObject>& Params) override;
+};
+
+/**
+ * Get the default value of a Blueprint variable
+ */
+class FECACommand_GetBlueprintVariableDefault : public IECACommand
+{
+public:
+	virtual FString GetName() const override { return TEXT("get_blueprint_variable_default"); }
+	virtual FString GetDescription() const override { return TEXT("Get the default value of a Blueprint variable"); }
+	virtual FString GetCategory() const override { return TEXT("Blueprint"); }
+	
+	virtual TArray<FECACommandParam> GetParameters() const override
+	{
+		return {
+			{ TEXT("blueprint_path"), TEXT("string"), TEXT("Path to the Blueprint asset"), true },
+			{ TEXT("variable_name"), TEXT("string"), TEXT("Name of the variable"), true }
+		};
+	}
+	
+	virtual FECACommandResult Execute(const TSharedPtr<FJsonObject>& Params) override;
+};
+
+/**
+ * Break a specific pin connection
+ */
+class FECACommand_BreakPinConnection : public IECACommand
+{
+public:
+	virtual FString GetName() const override { return TEXT("break_pin_connection"); }
+	virtual FString GetDescription() const override { return TEXT("Break a connection between two specific pins"); }
+	virtual FString GetCategory() const override { return TEXT("Blueprint Node"); }
+	
+	virtual TArray<FECACommandParam> GetParameters() const override
+	{
+		return {
+			{ TEXT("blueprint_path"), TEXT("string"), TEXT("Path to the Blueprint asset"), true },
+			{ TEXT("source_node_id"), TEXT("string"), TEXT("GUID of the source node"), true },
+			{ TEXT("source_pin"), TEXT("string"), TEXT("Name of the source pin"), true },
+			{ TEXT("target_node_id"), TEXT("string"), TEXT("GUID of the target node"), true },
+			{ TEXT("target_pin"), TEXT("string"), TEXT("Name of the target pin"), true },
+			{ TEXT("graph_name"), TEXT("string"), TEXT("Name of the graph"), false, TEXT("EventGraph") }
+		};
+	}
+	
+	virtual FECACommandResult Execute(const TSharedPtr<FJsonObject>& Params) override;
+};
+
+/**
+ * Add a component-bound event node (e.g., OnComponentBeginOverlap for a specific component)
+ */
+class FECACommand_AddComponentEventNode : public IECACommand
+{
+public:
+	virtual FString GetName() const override { return TEXT("add_component_event_node"); }
+	virtual FString GetDescription() const override { return TEXT("Add an event node bound to a specific component (e.g., OnComponentBeginOverlap, OnComponentEndOverlap, OnComponentHit)"); }
+	virtual FString GetCategory() const override { return TEXT("Blueprint Node"); }
+	
+	virtual TArray<FECACommandParam> GetParameters() const override
+	{
+		return {
+			{ TEXT("blueprint_path"), TEXT("string"), TEXT("Path to the Blueprint asset"), true },
+			{ TEXT("component_name"), TEXT("string"), TEXT("Name of the component to bind the event to"), true },
+			{ TEXT("event_name"), TEXT("string"), TEXT("Event name: OnComponentBeginOverlap, OnComponentEndOverlap, OnComponentHit, OnComponentWake, OnComponentSleep, etc."), true },
+			{ TEXT("node_position"), TEXT("object"), TEXT("Node position in graph {x, y}"), false }
+		};
+	}
+	
+	virtual FECACommandResult Execute(const TSharedPtr<FJsonObject>& Params) override;
+};
+
+/**
+ * Delete a component from a Blueprint
+ */
+class FECACommand_DeleteBlueprintComponent : public IECACommand
+{
+public:
+	virtual FString GetName() const override { return TEXT("delete_blueprint_component"); }
+	virtual FString GetDescription() const override { return TEXT("Delete a component from a Blueprint by name"); }
+	virtual FString GetCategory() const override { return TEXT("Blueprint"); }
+	
+	virtual TArray<FECACommandParam> GetParameters() const override
+	{
+		return {
+			{ TEXT("blueprint_path"), TEXT("string"), TEXT("Path to the Blueprint asset"), true },
+			{ TEXT("component_name"), TEXT("string"), TEXT("Name of the component to delete"), true }
+		};
+	}
+	
+	virtual FECACommandResult Execute(const TSharedPtr<FJsonObject>& Params) override;
+};
+
+/**
+ * Add a flow control node (Branch, Select, Switch) to a Blueprint
+ * 
+ * Flow control node types:
+ * - branch: If/Then/Else (UK2Node_IfThenElse) - has Condition input, True/False exec outputs
+ * - select: Select based on index (UK2Node_Select) - selects from array of options
+ * - switch_int: Switch on integer value
+ * - switch_string: Switch on string value  
+ * - switch_name: Switch on FName value
+ * - switch_enum: Switch on enum value (requires enum_type parameter)
+ * 
+ * Example: {"blueprint_path":"/Game/BP_Test", "node_type":"branch", "node_position":{"x":200,"y":100}}
+ */
+class FECACommand_AddBlueprintFlowControlNode : public IECACommand
+{
+public:
+	virtual FString GetName() const override { return TEXT("add_blueprint_flow_control_node"); }
+	virtual FString GetDescription() const override { return TEXT("Add a flow control node: branch (if/else), select, switch_int, switch_string, switch_name, switch_enum. Example: {\"node_type\":\"branch\"}"); }
+	virtual FString GetCategory() const override { return TEXT("Blueprint Node"); }
+	
+	virtual TArray<FECACommandParam> GetParameters() const override
+	{
+		return {
+			{ TEXT("blueprint_path"), TEXT("string"), TEXT("Path to the Blueprint asset"), true },
+			{ TEXT("node_type"), TEXT("string"), TEXT("Type: branch, select, switch_int, switch_string, switch_name, switch_enum"), true },
+			{ TEXT("enum_type"), TEXT("string"), TEXT("For switch_enum: the enum type path (e.g., /Script/Engine.ECollisionChannel)"), false },
+			{ TEXT("graph_name"), TEXT("string"), TEXT("Name of the graph to add to"), false, TEXT("EventGraph") },
+			{ TEXT("node_position"), TEXT("object"), TEXT("Node position in graph {x, y}"), false }
+		};
+	}
+	
+	virtual FECACommandResult Execute(const TSharedPtr<FJsonObject>& Params) override;
+};
+
+/**
+ * Auto-layout nodes in a Blueprint graph for visual clarity
+ * 
+ * Layout strategies:
+ * - horizontal: Arrange nodes left-to-right following execution flow (default)
+ * - vertical: Arrange nodes top-to-bottom following execution flow
+ * - tree: Arrange as a tree structure from root events
+ * - compact: Minimize total graph area while maintaining readability
+ * 
+ * The algorithm:
+ * 1. Identifies root nodes (events, entry points with no input exec connections)
+ * 2. Traverses execution flow to determine node order/depth
+ * 3. Positions nodes with consistent spacing to avoid overlaps
+ * 4. Aligns data-flow nodes (getters, pure functions) near their consumers
+ * 
+ * Example: {"blueprint_path":"/Game/BP_Test", "graph_name":"EventGraph", "strategy":"horizontal"}
+ */
+class FECACommand_AutoLayoutBlueprintGraph : public IECACommand
+{
+public:
+	virtual FString GetName() const override { return TEXT("auto_layout_blueprint_graph"); }
+	virtual FString GetDescription() const override { return TEXT("Automatically arrange nodes in a Blueprint graph for visual clarity. Organizes by execution flow with consistent spacing."); }
+	virtual FString GetCategory() const override { return TEXT("Blueprint Node"); }
+	
+	virtual TArray<FECACommandParam> GetParameters() const override
+	{
+		return {
+			{ TEXT("blueprint_path"), TEXT("string"), TEXT("Path to the Blueprint asset"), true },
+			{ TEXT("graph_name"), TEXT("string"), TEXT("Name of the graph to layout"), false, TEXT("EventGraph") },
+			{ TEXT("strategy"), TEXT("string"), TEXT("Layout strategy: horizontal (default), vertical, tree, compact"), false, TEXT("horizontal") },
+			{ TEXT("spacing_x"), TEXT("number"), TEXT("Horizontal spacing between nodes (default: 400)"), false, TEXT("400") },
+			{ TEXT("spacing_y"), TEXT("number"), TEXT("Vertical spacing between nodes (default: 150)"), false, TEXT("150") },
+			{ TEXT("align_comments"), TEXT("boolean"), TEXT("Also reposition comment nodes to wrap their contents"), false, TEXT("true") },
+			{ TEXT("selected_only"), TEXT("boolean"), TEXT("Only layout selected nodes (false = entire graph)"), false, TEXT("false") },
+			{ TEXT("node_ids"), TEXT("array"), TEXT("Specific node IDs to layout (optional, overrides selected_only)"), false }
+		};
+	}
+	
+	virtual FECACommandResult Execute(const TSharedPtr<FJsonObject>& Params) override;
+};
diff --git a/Engine/Plugins/Experimental/ECABridge/Source/ECABridge/Public/Commands/ECACommand.h b/Engine/Plugins/Experimental/ECABridge/Source/ECABridge/Public/Commands/ECACommand.h
new file mode 100644
--- /dev/null
+++ b/Engine/Plugins/Experimental/ECABridge/Source/ECABridge/Public/Commands/ECACommand.h
@@ -0,0 +1,163 @@
+// Copyright Epic Games, Inc. All Rights Reserved.
+
+#pragma once
+
+#include "CoreMinimal.h"
+#include "Dom/JsonObject.h"
+
+// Forward declarations
+class AActor;
+class UBlueprint;
+
+/**
+ * Result of command execution
+ */
+struct FECACommandResult
+{
+	bool bSuccess = false;
+	FString ErrorMessage;
+	TSharedPtr<FJsonObject> ResultData;
+	
+	static FECACommandResult Success(TSharedPtr<FJsonObject> Data = nullptr)
+	{
+		FECACommandResult Result;
+		Result.bSuccess = true;
+		Result.ResultData = Data;
+		return Result;
+	}
+	
+	static FECACommandResult Error(const FString& Message)
+	{
+		FECACommandResult Result;
+		Result.bSuccess = false;
+		Result.ErrorMessage = Message;
+		return Result;
+	}
+	
+	/** Convert to JSON response string */
+	FString ToJsonString() const;
+};
+
+/**
+ * Parameter definition for command documentation
+ */
+struct FECACommandParam
+{
+	FString Name;
+	FString Type; // "string", "number", "boolean", "object", "array"
+	FString Description;
+	bool bRequired = false;
+	FString DefaultValue;
+};
+
+/**
+ * Base class for all ECA commands
+ * 
+ * Subclass this to create new commands. Commands are automatically registered
+ * when the module loads via the REGISTER_ECA_COMMAND macro.
+ */
+class ECABRIDGE_API IECACommand
+{
+public:
+	virtual ~IECACommand() = default;
+	
+	/** Unique command name (e.g., "create_actor", "get_actors_in_level") */
+	virtual FString GetName() const = 0;
+	
+	/** Human-readable description for documentation */
+	virtual FString GetDescription() const = 0;
+	
+	/** Category for organization (e.g., "Actor", "Blueprint", "Editor") */
+	virtual FString GetCategory() const = 0;
+	
+	/** Parameter definitions for documentation */
+	virtual TArray<FECACommandParam> GetParameters() const { return {}; }
+	
+	/** Execute the command with given parameters */
+	virtual FECACommandResult Execute(const TSharedPtr<FJsonObject>& Params) = 0;
+	
+protected:
+	// Helper methods for parameter extraction
+	
+	static bool GetStringParam(const TSharedPtr<FJsonObject>& Params, const FString& Name, FString& OutValue, bool bRequired = true);
+	static bool GetIntParam(const TSharedPtr<FJsonObject>& Params, const FString& Name, int32& OutValue, bool bRequired = true);
+	static bool GetFloatParam(const TSharedPtr<FJsonObject>& Params, const FString& Name, double& OutValue, bool bRequired = true);
+	static bool GetBoolParam(const TSharedPtr<FJsonObject>& Params, const FString& Name, bool& OutValue, bool bRequired = true);
+	static bool GetVectorParam(const TSharedPtr<FJsonObject>& Params, const FString& Name, FVector& OutValue, bool bRequired = true);
+	static bool GetRotatorParam(const TSharedPtr<FJsonObject>& Params, const FString& Name, FRotator& OutValue, bool bRequired = true);
+	static bool GetArrayParam(const TSharedPtr<FJsonObject>& Params, const FString& Name, const TArray<TSharedPtr<FJsonValue>>*& OutValue, bool bRequired = true);
+	static bool GetObjectParam(const TSharedPtr<FJsonObject>& Params, const FString& Name, const TSharedPtr<FJsonObject>*& OutValue, bool bRequired = true);
+	
+	// Helper methods for result building
+	
+	static TSharedPtr<FJsonObject> MakeResult();
+	static TSharedPtr<FJsonObject> VectorToJson(const FVector& Vec);
+	static TSharedPtr<FJsonObject> RotatorToJson(const FRotator& Rot);
+	static TSharedPtr<FJsonObject> TransformToJson(const FTransform& Transform);
+	
+	// Helper methods for common operations
+	
+	static UWorld* GetEditorWorld();
+	static AActor* FindActorByName(const FString& ActorName);
+	static UBlueprint* LoadBlueprintByPath(const FString& BlueprintPath);
+};
+
+/**
+ * Command registry - manages all available commands
+ */
+class ECABRIDGE_API FECACommandRegistry
+{
+public:
+	/** Get the singleton instance */
+	static FECACommandRegistry& Get();
+	
+	/** Register a command */
+	void RegisterCommand(TSharedPtr<IECACommand> Command);
+	
+	/** Unregister a command by name */
+	void UnregisterCommand(const FString& Name);
+	
+	/** Get a command by name */
+	TSharedPtr<IECACommand> GetCommand(const FString& Name) const;
+	
+	/** Get all registered commands */
+	TArray<TSharedPtr<IECACommand>> GetAllCommands() const;
+	
+	/** Get commands by category */
+	TArray<TSharedPtr<IECACommand>> GetCommandsByCategory(const FString& Category) const;
+	
+	/** Check if a command exists */
+	bool HasCommand(const FString& Name) const;
+	
+	/** Get all category names */
+	TArray<FString> GetCategories() const;
+	
+	/** Execute a command by name */
+	FECACommandResult ExecuteCommand(const FString& Name, const TSharedPtr<FJsonObject>& Params);
+	
+	/** Log all registered commands */
+	void LogCommands() const;
+
+private:
+	FECACommandRegistry() = default;
+	
+	TMap<FString, TSharedPtr<IECACommand>> Commands;
+	mutable FCriticalSection CommandsLock;
+};
+
+/**
+ * Helper class for automatic command registration
+ */
+template<typename T>
+class TECACommandRegistrar
+{
+public:
+	TECACommandRegistrar()
+	{
+		FECACommandRegistry::Get().RegisterCommand(MakeShared<T>());
+	}
+};
+
+/** Macro for easy command registration */
+#define REGISTER_ECA_COMMAND(CommandClass) \
+	static TECACommandRegistrar<CommandClass> CommandClass##_Registrar;
