ECABridge experimental MCP plugin

diff --git a/Engine/Plugins/Experimental/ECABridge/Source/ECABridge/Public/Commands/ECAComponentCommands.h b/Engine/Plugins/Experimental/ECABridge/Source/ECABridge/Public/Commands/ECAComponentCommands.h
new file mode 100644
--- /dev/null
+++ b/Engine/Plugins/Experimental/ECABridge/Source/ECABridge/Public/Commands/ECAComponentCommands.h
@@ -0,0 +1,120 @@
+// Copyright Epic Games, Inc. All Rights Reserved.
+
+#pragma once
+
+#include "Commands/ECACommand.h"
+
+/**
+ * Set static mesh properties on a component
+ */
+class FECACommand_SetStaticMeshProperties : public IECACommand
+{
+public:
+	virtual FString GetName() const override { return TEXT("set_static_mesh_properties"); }
+	virtual FString GetDescription() const override { return TEXT("Set static mesh on a StaticMeshComponent in a Blueprint"); }
+	virtual FString GetCategory() const override { return TEXT("Component"); }
+	
+	virtual TArray<FECACommandParam> GetParameters() const override
+	{
+		return {
+			{ TEXT("blueprint_path"), TEXT("string"), TEXT("Path to the Blueprint asset"), true },
+			{ TEXT("component_name"), TEXT("string"), TEXT("Name of the StaticMeshComponent"), true },
+			{ TEXT("static_mesh"), TEXT("string"), TEXT("Path to the static mesh asset"), true }
+		};
+	}
+	
+	virtual FECACommandResult Execute(const TSharedPtr<FJsonObject>& Params) override;
+};
+
+/**
+ * Set physics properties on a component
+ */
+class FECACommand_SetPhysicsProperties : public IECACommand
+{
+public:
+	virtual FString GetName() const override { return TEXT("set_physics_properties"); }
+	virtual FString GetDescription() const override { return TEXT("Set physics properties on a component in a Blueprint"); }
+	virtual FString GetCategory() const override { return TEXT("Component"); }
+	
+	virtual TArray<FECACommandParam> GetParameters() const override
+	{
+		return {
+			{ TEXT("blueprint_path"), TEXT("string"), TEXT("Path to the Blueprint asset"), true },
+			{ TEXT("component_name"), TEXT("string"), TEXT("Name of the component"), true },
+			{ TEXT("simulate_physics"), TEXT("boolean"), TEXT("Enable physics simulation"), false, TEXT("true") },
+			{ TEXT("gravity_enabled"), TEXT("boolean"), TEXT("Enable gravity"), false, TEXT("true") },
+			{ TEXT("mass"), TEXT("number"), TEXT("Mass in kg"), false, TEXT("1.0") },
+			{ TEXT("linear_damping"), TEXT("number"), TEXT("Linear damping"), false, TEXT("0.01") },
+			{ TEXT("angular_damping"), TEXT("number"), TEXT("Angular damping"), false, TEXT("0.0") }
+		};
+	}
+	
+	virtual FECACommandResult Execute(const TSharedPtr<FJsonObject>& Params) override;
+};
+
+/**
+ * Set a generic property on a component
+ */
+class FECACommand_SetComponentProperty : public IECACommand
+{
+public:
+	virtual FString GetName() const override { return TEXT("set_component_property"); }
+	virtual FString GetDescription() const override { return TEXT("Set a property on a component in a Blueprint"); }
+	virtual FString GetCategory() const override { return TEXT("Component"); }
+	
+	virtual TArray<FECACommandParam> GetParameters() const override
+	{
+		return {
+			{ TEXT("blueprint_path"), TEXT("string"), TEXT("Path to the Blueprint asset"), true },
+			{ TEXT("component_name"), TEXT("string"), TEXT("Name of the component"), true },
+			{ TEXT("property_name"), TEXT("string"), TEXT("Name of the property to set"), true },
+			{ TEXT("property_value"), TEXT("any"), TEXT("Value to set (type depends on property)"), true }
+		};
+	}
+	
+	virtual FECACommandResult Execute(const TSharedPtr<FJsonObject>& Params) override;
+};
+
+/**
+ * Set component transform (relative)
+ */
+class FECACommand_SetComponentTransform : public IECACommand
+{
+public:
+	virtual FString GetName() const override { return TEXT("set_component_transform"); }
+	virtual FString GetDescription() const override { return TEXT("Set the relative transform of a component in a Blueprint"); }
+	virtual FString GetCategory() const override { return TEXT("Component"); }
+	
+	virtual TArray<FECACommandParam> GetParameters() const override
+	{
+		return {
+			{ TEXT("blueprint_path"), TEXT("string"), TEXT("Path to the Blueprint asset"), true },
+			{ TEXT("component_name"), TEXT("string"), TEXT("Name of the component"), true },
+			{ TEXT("location"), TEXT("object"), TEXT("Relative location {x, y, z}"), false },
+			{ TEXT("rotation"), TEXT("object"), TEXT("Relative rotation {pitch, yaw, roll}"), false },
+			{ TEXT("scale"), TEXT("object"), TEXT("Scale {x, y, z}"), false }
+		};
+	}
+	
+	virtual FECACommandResult Execute(const TSharedPtr<FJsonObject>& Params) override;
+};
+
+/**
+ * Get components in a Blueprint
+ */
+class FECACommand_GetBlueprintComponents : public IECACommand
+{
+public:
+	virtual FString GetName() const override { return TEXT("get_blueprint_components"); }
+	virtual FString GetDescription() const override { return TEXT("Get all components in a Blueprint"); }
+	virtual FString GetCategory() const override { return TEXT("Component"); }
+	
+	virtual TArray<FECACommandParam> GetParameters() const override
+	{
+		return {
+			{ TEXT("blueprint_path"), TEXT("string"), TEXT("Path to the Blueprint asset"), true }
+		};
+	}
+	
+	virtual FECACommandResult Execute(const TSharedPtr<FJsonObject>& Params) override;
+};
diff --git a/Engine/Plugins/Experimental/ECABridge/Source/ECABridge/Public/Commands/ECADataTableCommands.h b/Engine/Plugins/Experimental/ECABridge/Source/ECABridge/Public/Commands/ECADataTableCommands.h
new file mode 100644
--- /dev/null
+++ b/Engine/Plugins/Experimental/ECABridge/Source/ECABridge/Public/Commands/ECADataTableCommands.h
@@ -0,0 +1,91 @@
+// Copyright Epic Games, Inc. All Rights Reserved.
+
+#pragma once
+
+#include "Commands/ECACommand.h"
+
+/**
+ * Get DataTable schema - list all columns/fields and their types
+ */
+class FECACommand_GetDataTableSchema : public IECACommand
+{
+public:
+	virtual FString GetName() const override { return TEXT("get_datatable_schema"); }
+	virtual FString GetDescription() const override { return TEXT("Get the row structure (schema) of a DataTable, including all column names and types"); }
+	virtual FString GetCategory() const override { return TEXT("DataTable"); }
+	
+	virtual TArray<FECACommandParam> GetParameters() const override
+	{
+		return {
+			{ TEXT("asset_path"), TEXT("string"), TEXT("Content path to the DataTable asset (e.g., /Game/Data/MyDataTable)"), true },
+		};
+	}
+	
+	virtual FECACommandResult Execute(const TSharedPtr<FJsonObject>& Params) override;
+};
+
+/**
+ * Get DataTable rows - read all or specific rows from a DataTable
+ */
+class FECACommand_GetDataTableRows : public IECACommand
+{
+public:
+	virtual FString GetName() const override { return TEXT("get_datatable_rows"); }
+	virtual FString GetDescription() const override { return TEXT("Get rows from a DataTable. Returns all rows or specific rows by name. Each row includes all column values."); }
+	virtual FString GetCategory() const override { return TEXT("DataTable"); }
+	
+	virtual TArray<FECACommandParam> GetParameters() const override
+	{
+		return {
+			{ TEXT("asset_path"), TEXT("string"), TEXT("Content path to the DataTable asset (e.g., /Game/Data/MyDataTable)"), true },
+			{ TEXT("row_names"), TEXT("array"), TEXT("Optional array of specific row names to retrieve. If omitted, returns all rows."), false },
+			{ TEXT("limit"), TEXT("number"), TEXT("Maximum number of rows to return (default: all)"), false },
+			{ TEXT("offset"), TEXT("number"), TEXT("Number of rows to skip (for pagination, default: 0)"), false },
+		};
+	}
+	
+	virtual FECACommandResult Execute(const TSharedPtr<FJsonObject>& Params) override;
+};
+
+/**
+ * Set DataTable row - add or update a row in a DataTable
+ */
+class FECACommand_SetDataTableRow : public IECACommand
+{
+public:
+	virtual FString GetName() const override { return TEXT("set_datatable_row"); }
+	virtual FString GetDescription() const override { return TEXT("Add a new row or update an existing row in a DataTable. Provide the row name and column values as a JSON object."); }
+	virtual FString GetCategory() const override { return TEXT("DataTable"); }
+	
+	virtual TArray<FECACommandParam> GetParameters() const override
+	{
+		return {
+			{ TEXT("asset_path"), TEXT("string"), TEXT("Content path to the DataTable asset"), true },
+			{ TEXT("row_name"), TEXT("string"), TEXT("Name/key of the row to add or update"), true },
+			{ TEXT("row_data"), TEXT("object"), TEXT("JSON object with column names as keys and values to set. Use get_datatable_schema to discover column names."), true },
+		};
+	}
+	
+	virtual FECACommandResult Execute(const TSharedPtr<FJsonObject>& Params) override;
+};
+
+/**
+ * Delete DataTable row - remove a row from a DataTable
+ */
+class FECACommand_DeleteDataTableRow : public IECACommand
+{
+public:
+	virtual FString GetName() const override { return TEXT("delete_datatable_row"); }
+	virtual FString GetDescription() const override { return TEXT("Delete a row from a DataTable by its row name/key"); }
+	virtual FString GetCategory() const override { return TEXT("DataTable"); }
+	
+	virtual TArray<FECACommandParam> GetParameters() const override
+	{
+		return {
+			{ TEXT("asset_path"), TEXT("string"), TEXT("Content path to the DataTable asset"), true },
+			{ TEXT("row_name"), TEXT("string"), TEXT("Name/key of the row to delete"), true },
+		};
+	}
+	
+	virtual FECACommandResult Execute(const TSharedPtr<FJsonObject>& Params) override;
+};
diff --git a/Engine/Plugins/Experimental/ECABridge/Source/ECABridge/Public/Commands/ECAEditorCommands.h b/Engine/Plugins/Experimental/ECABridge/Source/ECABridge/Public/Commands/ECAEditorCommands.h
new file mode 100644
--- /dev/null
+++ b/Engine/Plugins/Experimental/ECABridge/Source/ECABridge/Public/Commands/ECAEditorCommands.h
@@ -0,0 +1,225 @@
+// Copyright Epic Games, Inc. All Rights Reserved.
+
+#pragma once
+
+#include "Commands/ECACommand.h"
+
+/**
+ * Focus viewport on actor(s)
+ */
+class FECACommand_FocusViewport : public IECACommand
+{
+public:
+	virtual FString GetName() const override { return TEXT("focus_viewport"); }
+	virtual FString GetDescription() const override { return TEXT("Focus the editor viewport on one or more actors"); }
+	virtual FString GetCategory() const override { return TEXT("Editor"); }
+	
+	virtual TArray<FECACommandParam> GetParameters() const override
+	{
+		return {
+			{ TEXT("actor_names"), TEXT("array"), TEXT("Array of actor names to focus on"), false },
+			{ TEXT("actor_name"), TEXT("string"), TEXT("Single actor name to focus on"), false },
+			{ TEXT("location"), TEXT("object"), TEXT("Location to focus on {x, y, z}"), false },
+			{ TEXT("distance"), TEXT("number"), TEXT("Camera distance from target"), false, TEXT("500") }
+		};
+	}
+	
+	virtual FECACommandResult Execute(const TSharedPtr<FJsonObject>& Params) override;
+};
+
+/**
+ * Select actors in the editor
+ */
+class FECACommand_SelectActors : public IECACommand
+{
+public:
+	virtual FString GetName() const override { return TEXT("select_actors"); }
+	virtual FString GetDescription() const override { return TEXT("Select actors in the editor"); }
+	virtual FString GetCategory() const override { return TEXT("Editor"); }
+	
+	virtual TArray<FECACommandParam> GetParameters() const override
+	{
+		return {
+			{ TEXT("actor_names"), TEXT("array"), TEXT("Array of actor names to select"), false },
+			{ TEXT("actor_name"), TEXT("string"), TEXT("Single actor name to select"), false },
+			{ TEXT("add_to_selection"), TEXT("boolean"), TEXT("Add to current selection instead of replacing"), false, TEXT("false") },
+			{ TEXT("deselect_all"), TEXT("boolean"), TEXT("Deselect all actors"), false, TEXT("false") }
+		};
+	}
+	
+	virtual FECACommandResult Execute(const TSharedPtr<FJsonObject>& Params) override;
+};
+
+/**
+ * Get selected actors
+ */
+class FECACommand_GetSelectedActors : public IECACommand
+{
+public:
+	virtual FString GetName() const override { return TEXT("get_selected_actors"); }
+	virtual FString GetDescription() const override { return TEXT("Get the currently selected actors"); }
+	virtual FString GetCategory() const override { return TEXT("Editor"); }
+	
+	virtual FECACommandResult Execute(const TSharedPtr<FJsonObject>& Params) override;
+};
+
+/**
+ * Take a depth screenshot of the viewport
+ */
+class FECACommand_TakeDepthScreenshot : public IECACommand
+{
+public:
+	virtual FString GetName() const override { return TEXT("take_depth_screenshot"); }
+	virtual FString GetDescription() const override { return TEXT("Take a depth buffer screenshot from the editor viewport. Returns grayscale image where darker=closer, lighter=farther."); }
+	virtual FString GetCategory() const override { return TEXT("Editor"); }
+	
+	virtual TArray<FECACommandParam> GetParameters() const override
+	{
+		return {
+			{ TEXT("file_path"), TEXT("string"), TEXT("Path to save the depth image (if not specified, returns base64)"), false },
+			{ TEXT("width"), TEXT("number"), TEXT("Image width"), false, TEXT("1024") },
+			{ TEXT("height"), TEXT("number"), TEXT("Image height"), false, TEXT("1024") },
+			{ TEXT("max_depth"), TEXT("number"), TEXT("Maximum depth in units (default 10000 = 100m). Depths beyond this are clamped to white."), false, TEXT("10000") }
+		};
+	}
+	
+	virtual FECACommandResult Execute(const TSharedPtr<FJsonObject>& Params) override;
+};
+
+/**
+ * Take a screenshot of the active viewport, game window, or editor window.
+ * Supports three capture targets:
+ *   "pie"             - Game window via FSlateApplication::TakeScreenshot (full composited with all UI)
+ *   "editor_viewport" - Editor 3D viewport via GetViewportScreenShot (render target only, no Slate)
+ *   "editor_window"   - Full editor window via FSlateApplication::TakeScreenshot (all panels, tabs, graphs)
+ * Auto-detects: if PIE is running uses "pie", otherwise uses "editor_viewport".
+ */
+class FECACommand_TakeGameplayScreenshot : public IECACommand
+{
+public:
+	virtual FString GetName() const override { return TEXT("take_gameplay_screenshot"); }
+	virtual FString GetDescription() const override { return TEXT("Take a screenshot of the game viewport, editor viewport, or the full editor window. By default auto-detects: captures the PIE game window (with all Slate/UMG/HUD overlays) if playing, otherwise captures the editor 3D viewport. Use the 'target' parameter to explicitly choose what to capture, including 'editor_window' for the full editor UI (Blueprint graphs, widget hierarchies, Details panels, etc)."); }
+	virtual FString GetCategory() const override { return TEXT("Editor"); }
+	
+	virtual TArray<FECACommandParam> GetParameters() const override
+	{
+		return {
+			{ TEXT("file_path"), TEXT("string"), TEXT("Path to save the screenshot PNG. If not specified, saves to project Saved/Screenshots/ with a timestamp filename."), false },
+			{ TEXT("target"), TEXT("string"), TEXT("What to capture: 'pie' (game window with UI overlays), 'editor_viewport' (3D viewport only), or 'editor_window' (full editor UI including all panels/tabs). If not specified, auto-detects: PIE window if playing, otherwise editor viewport."), false }
+		};
+	}
+	
+	virtual FECACommandResult Execute(const TSharedPtr<FJsonObject>& Params) override;
+};
+
+/**
+ * Run a console command
+ */
+class FECACommand_RunConsoleCommand : public IECACommand
+{
+public:
+	virtual FString GetName() const override { return TEXT("run_console_command"); }
+	virtual FString GetDescription() const override { return TEXT("Execute an Unreal console command"); }
+	virtual FString GetCategory() const override { return TEXT("Editor"); }
+	
+	virtual TArray<FECACommandParam> GetParameters() const override
+	{
+		return {
+			{ TEXT("command"), TEXT("string"), TEXT("The console command to execute"), true }
+		};
+	}
+	
+	virtual FECACommandResult Execute(const TSharedPtr<FJsonObject>& Params) override;
+};
+
+/**
+ * Get current level info
+ */
+class FECACommand_GetLevelInfo : public IECACommand
+{
+public:
+	virtual FString GetName() const override { return TEXT("get_level_info"); }
+	virtual FString GetDescription() const override { return TEXT("Get information about the current level"); }
+	virtual FString GetCategory() const override { return TEXT("Editor"); }
+	
+	virtual FECACommandResult Execute(const TSharedPtr<FJsonObject>& Params) override;
+};
+
+/**
+ * Open a level
+ */
+class FECACommand_OpenLevel : public IECACommand
+{
+public:
+	virtual FString GetName() const override { return TEXT("open_level"); }
+	virtual FString GetDescription() const override { return TEXT("Open a level in the editor"); }
+	virtual FString GetCategory() const override { return TEXT("Editor"); }
+	
+	virtual TArray<FECACommandParam> GetParameters() const override
+	{
+		return {
+			{ TEXT("level_path"), TEXT("string"), TEXT("Path to the level asset"), true }
+		};
+	}
+	
+	virtual FECACommandResult Execute(const TSharedPtr<FJsonObject>& Params) override;
+};
+
+/**
+ * Save current level
+ */
+class FECACommand_SaveLevel : public IECACommand
+{
+public:
+	virtual FString GetName() const override { return TEXT("save_level"); }
+	virtual FString GetDescription() const override { return TEXT("Save the current level"); }
+	virtual FString GetCategory() const override { return TEXT("Editor"); }
+	
+	virtual FECACommandResult Execute(const TSharedPtr<FJsonObject>& Params) override;
+};
+
+/**
+ * Play in Editor
+ */
+class FECACommand_PlayInEditor : public IECACommand
+{
+public:
+	virtual FString GetName() const override { return TEXT("play_in_editor"); }
+	virtual FString GetDescription() const override { return TEXT("Start Play in Editor session"); }
+	virtual FString GetCategory() const override { return TEXT("Editor"); }
+	
+	virtual TArray<FECACommandParam> GetParameters() const override
+	{
+		return {
+			{ TEXT("viewport_index"), TEXT("number"), TEXT("Viewport index to play in"), false, TEXT("0") }
+		};
+	}
+	
+	virtual FECACommandResult Execute(const TSharedPtr<FJsonObject>& Params) override;
+};
+
+/**
+ * Stop Play in Editor
+ */
+class FECACommand_StopPlayInEditor : public IECACommand
+{
+public:
+	virtual FString GetName() const override { return TEXT("stop_play_in_editor"); }
+	virtual FString GetDescription() const override { return TEXT("Stop the current Play in Editor session"); }
+	virtual FString GetCategory() const override { return TEXT("Editor"); }
+	
+	virtual FECACommandResult Execute(const TSharedPtr<FJsonObject>& Params) override;
+};
+
+/**
+ * Get editor play state
+ */
+class FECACommand_GetPlayState : public IECACommand
+{
+public:
+	virtual FString GetName() const override { return TEXT("get_play_state"); }
+	virtual FString GetDescription() const override { return TEXT("Get the current Play in Editor state"); }
+	virtual FString GetCategory() const override { return TEXT("Editor"); }
+	
+	virtual FECACommandResult Execute(const TSharedPtr<FJsonObject>& Params) override;
+};
diff --git a/Engine/Plugins/Experimental/ECABridge/Source/ECABridge/Public/Commands/ECAListItemData.h b/Engine/Plugins/Experimental/ECABridge/Source/ECABridge/Public/Commands/ECAListItemData.h
new file mode 100644
--- /dev/null
+++ b/Engine/Plugins/Experimental/ECABridge/Source/ECABridge/Public/Commands/ECAListItemData.h
@@ -0,0 +1,27 @@
+// Copyright Epic Games, Inc. All Rights Reserved.
+
+#pragma once
+
+#include "CoreMinimal.h"
+#include "UObject/Object.h"
+#include "ECAListItemData.generated.h"
+
+/**
+ * Simple UObject wrapper for list item data.
+ * UListView requires UObject* items - this wraps arbitrary JSON data into a UObject
+ * that can be passed to the list and retrieved by entry widgets via IUserObjectListEntry.
+ */
+UCLASS(Transient, NotBlueprintable, MinimalAPI)
+class UECAListItemData : public UObject
+{
+	GENERATED_BODY()
+
+public:
+	/** JSON string containing the item's data payload */
+	UPROPERTY()
+	FString ItemDataJson;
+
+	/** Index of this item in the list */
+	UPROPERTY()
+	int32 ItemIndex = 0;
+};
diff --git a/Engine/Plugins/Experimental/ECABridge/Source/ECABridge/Public/Commands/ECAMVVMCommands.h b/Engine/Plugins/Experimental/ECABridge/Source/ECABridge/Public/Commands/ECAMVVMCommands.h
new file mode 100644
--- /dev/null
+++ b/Engine/Plugins/Experimental/ECABridge/Source/ECABridge/Public/Commands/ECAMVVMCommands.h
@@ -0,0 +1,186 @@
+// Copyright Epic Games, Inc. All Rights Reserved.
+
+#pragma once
+
+#include "Commands/ECACommand.h"
+
+/**
+ * Assign a ViewModel class to a Widget Blueprint.
+ */
+class FECACommand_SetWidgetViewmodel : public IECACommand
+{
+public:
+	virtual FString GetName() const override { return TEXT("set_widget_viewmodel"); }
+	virtual FString GetDescription() const override { return TEXT("Assign a ViewModel class to a Widget Blueprint for MVVM binding"); }
+	virtual FString GetCategory() const override { return TEXT("MVVM"); }
+
+	virtual TArray<FECACommandParam> GetParameters() const override
+	{
+		return {
+			{ TEXT("widget_path"), TEXT("string"), TEXT("Path to the Widget Blueprint asset"), true },
+			{ TEXT("viewmodel_class"), TEXT("string"), TEXT("Name of the ViewModel class to assign"), true }
+		};
+	}
+
+	virtual FECACommandResult Execute(const TSharedPtr<FJsonObject>& Params) override;
+};
+
+/**
+ * Bind a TextBlock or RichTextBlock's Text property to a ViewModel property.
+ */
+class FECACommand_BindTextToViewmodel : public IECACommand
+{
+public:
+	virtual FString GetName() const override { return TEXT("bind_text_to_viewmodel"); }
+	virtual FString GetDescription() const override { return TEXT("Bind a text widget's Text property to a ViewModel property"); }
+	virtual FString GetCategory() const override { return TEXT("MVVM"); }
+
+	virtual TArray<FECACommandParam> GetParameters() const override
+	{
+		return {
+			{ TEXT("widget_path"), TEXT("string"), TEXT("Path to the Widget Blueprint asset"), true },
+			{ TEXT("widget_name"), TEXT("string"), TEXT("Name of the TextBlock widget to bind"), true },
+			{ TEXT("viewmodel_property"), TEXT("string"), TEXT("Name of the ViewModel property to bind from"), true }
+		};
+	}
+
+	virtual FECACommandResult Execute(const TSharedPtr<FJsonObject>& Params) override;
+};
+
+/**
+ * Bind a widget's Visibility to a ViewModel bool using Conv_BoolToSlateVisibility.
+ * Supports all 20 true/false ESlateVisibility combinations.
+ */
+class FECACommand_BindVisibilityToViewmodel : public IECACommand
+{
+public:
+	virtual FString GetName() const override { return TEXT("bind_visibility_to_viewmodel"); }
+	virtual FString GetDescription() const override { return TEXT("Bind a widget's Visibility to a ViewModel bool via Conv_BoolToSlateVisibility with configurable true/false visibility values"); }
+	virtual FString GetCategory() const override { return TEXT("MVVM"); }
+
+	virtual TArray<FECACommandParam> GetParameters() const override
+	{
+		return {
+			{ TEXT("widget_path"), TEXT("string"), TEXT("Path to the Widget Blueprint asset"), true },
+			{ TEXT("widget_name"), TEXT("string"), TEXT("Name of the widget to bind"), true },
+			{ TEXT("viewmodel_property"), TEXT("string"), TEXT("Name of the ViewModel bool property"), true },
+			{ TEXT("true_visibility"), TEXT("string"), TEXT("Visibility when bool is true (Visible, Collapsed, Hidden, HitTestInvisible, SelfHitTestInvisible)"), false, TEXT("Visible") },
+			{ TEXT("false_visibility"), TEXT("string"), TEXT("Visibility when bool is false"), false, TEXT("Collapsed") }
+		};
+	}
+
+	virtual FECACommandResult Execute(const TSharedPtr<FJsonObject>& Params) override;
+};
+
+/**
+ * Bind an Image widget's Brush property to a ViewModel FSlateBrush property.
+ */
+class FECACommand_BindImageToViewmodel : public IECACommand
+{
+public:
+	virtual FString GetName() const override { return TEXT("bind_image_to_viewmodel"); }
+	virtual FString GetDescription() const override { return TEXT("Bind an Image widget's Brush to a ViewModel FSlateBrush property"); }
+	virtual FString GetCategory() const override { return TEXT("MVVM"); }
+
+	virtual TArray<FECACommandParam> GetParameters() const override
+	{
+		return {
+			{ TEXT("widget_path"), TEXT("string"), TEXT("Path to the Widget Blueprint asset"), true },
+			{ TEXT("widget_name"), TEXT("string"), TEXT("Name of the Image widget to bind"), true },
+			{ TEXT("viewmodel_property"), TEXT("string"), TEXT("Name of the ViewModel FSlateBrush property"), true }
+		};
+	}
+
+	virtual FECACommandResult Execute(const TSharedPtr<FJsonObject>& Params) override;
+};
+
+/**
+ * Query all ViewModels and MVVM bindings on a Widget Blueprint.
+ */
+class FECACommand_GetViewmodelBindings : public IECACommand
+{
+public:
+	virtual FString GetName() const override { return TEXT("get_viewmodel_bindings"); }
+	virtual FString GetDescription() const override { return TEXT("Get all ViewModel assignments and MVVM bindings on a Widget Blueprint"); }
+	virtual FString GetCategory() const override { return TEXT("MVVM"); }
+
+	virtual TArray<FECACommandParam> GetParameters() const override
+	{
+		return {
+			{ TEXT("widget_path"), TEXT("string"), TEXT("Path to the Widget Blueprint asset"), true }
+		};
+	}
+
+	virtual FECACommandResult Execute(const TSharedPtr<FJsonObject>& Params) override;
+};
+
+/**
+ * Remove an MVVM binding by index or by source/destination property match.
+ */
+class FECACommand_RemoveViewmodelBinding : public IECACommand
+{
+public:
+	virtual FString GetName() const override { return TEXT("remove_viewmodel_binding"); }
+	virtual FString GetDescription() const override { return TEXT("Remove an MVVM binding by index or by matching source/destination property names"); }
+	virtual FString GetCategory() const override { return TEXT("MVVM"); }
+
+	virtual TArray<FECACommandParam> GetParameters() const override
+	{
+		return {
+			{ TEXT("widget_path"), TEXT("string"), TEXT("Path to the Widget Blueprint asset"), true },
+			{ TEXT("binding_index"), TEXT("number"), TEXT("Index of the binding to remove (use get_viewmodel_bindings to find indices)"), false },
+			{ TEXT("source_property"), TEXT("string"), TEXT("Source ViewModel property name to match for removal"), false },
+			{ TEXT("destination_property"), TEXT("string"), TEXT("Destination widget property name to match for removal"), false }
+		};
+	}
+
+	virtual FECACommandResult Execute(const TSharedPtr<FJsonObject>& Params) override;
+};
+
+/**
+ * Add a generic MVVM binding between a ViewModel property and any widget property.
+ * Supports self-bindings, multi-ViewModel lookup by name, and optional conversion functions.
+ */
+class FECACommand_AddViewmodelBinding : public IECACommand
+{
+public:
+	virtual FString GetName() const override { return TEXT("add_viewmodel_binding"); }
+	virtual FString GetDescription() const override { return TEXT("Add a generic MVVM binding from a ViewModel property to any widget property, with optional conversion function"); }
+	virtual FString GetCategory() const override { return TEXT("MVVM"); }
+
+	virtual TArray<FECACommandParam> GetParameters() const override
+	{
+		return {
+			{ TEXT("widget_path"), TEXT("string"), TEXT("Path to the Widget Blueprint asset"), true },
+			{ TEXT("source_property"), TEXT("string"), TEXT("ViewModel property name to bind from"), true },
+			{ TEXT("destination_property"), TEXT("string"), TEXT("Destination property or setter name on the target widget"), true },
+			{ TEXT("destination_widget"), TEXT("string"), TEXT("Child widget name, or 'self' to bind to the widget itself (default: self)"), false, TEXT("self") },
+			{ TEXT("viewmodel_name"), TEXT("string"), TEXT("ViewModel name to bind from (default: first ViewModel)"), false },
+			{ TEXT("binding_mode"), TEXT("string"), TEXT("Binding mode: OneWayToDestination, TwoWay, OneWayToSource, OneTimeToDestination (default: OneWayToDestination)"), false, TEXT("OneWayToDestination") },
+			{ TEXT("conversion_function"), TEXT("string"), TEXT("Optional conversion function (e.g. 'ClassName::FuncName' or 'FuncName' for BP-local)"), false }
+		};
+	}
+
+	virtual FECACommandResult Execute(const TSharedPtr<FJsonObject>& Params) override;
+};
+
+/**
+ * List all bindable properties on a specific widget within a Widget Blueprint.
+ */
+class FECACommand_GetWidgetMVVMProperties : public IECACommand
+{
+public:
+	virtual FString GetName() const override { return TEXT("get_widget_mvvm_properties"); }
+	virtual FString GetDescription() const override { return TEXT("Get all bindable properties on a widget within a Widget Blueprint"); }
+	virtual FString GetCategory() const override { return TEXT("MVVM"); }
+
+	virtual TArray<FECACommandParam> GetParameters() const override
+	{
+		return {
+			{ TEXT("widget_path"), TEXT("string"), TEXT("Path to the Widget Blueprint asset"), true },
+			{ TEXT("widget_name"), TEXT("string"), TEXT("Name of the widget to inspect"), true }
+		};
+	}
+
+	virtual FECACommandResult Execute(const TSharedPtr<FJsonObject>& Params) override;
+};
diff --git a/Engine/Plugins/Experimental/ECABridge/Source/ECABridge/Public/Commands/ECAMaterialNodeCommands.h b/Engine/Plugins/Experimental/ECABridge/Source/ECABridge/Public/Commands/ECAMaterialNodeCommands.h
new file mode 100644
--- /dev/null
+++ b/Engine/Plugins/Experimental/ECABridge/Source/ECABridge/Public/Commands/ECAMaterialNodeCommands.h
@@ -0,0 +1,380 @@
+// Copyright Epic Games, Inc. All Rights Reserved.
+
+#pragma once
+
+#include "Commands/ECACommand.h"
+
+/**
+ * Add a material expression node to a material graph
+ */
+class FECACommand_AddMaterialNode : public IECACommand
+{
+public:
+	virtual FString GetName() const override { return TEXT("add_material_node"); }
+	virtual FString GetDescription() const override { return TEXT("Add a material expression node to a material graph"); }
+	virtual FString GetCategory() const override { return TEXT("Material Node"); }
+	
+	virtual TArray<FECACommandParam> GetParameters() const override
+	{
+		return {
+			{ TEXT("material_path"), TEXT("string"), TEXT("Path to the Material asset"), true },
+			{ TEXT("node_type"), TEXT("string"), TEXT("Type of node: TextureSample, TextureSampleParameter2D, Constant, Constant3Vector, Constant4Vector, VectorParameter, ScalarParameter, Multiply, Add, Lerp, ComponentMask, Append, etc."), true },
+			{ TEXT("node_position"), TEXT("object"), TEXT("Node position in graph {x, y}"), false },
+			{ TEXT("parameter_name"), TEXT("string"), TEXT("Parameter name (for parameter types)"), false },
+			{ TEXT("default_value"), TEXT("any"), TEXT("Default value (scalar, vector {r,g,b,a}, or texture path)"), false }
+		};
+	}
+	
+	virtual FECACommandResult Execute(const TSharedPtr<FJsonObject>& Params) override;
+};
+
+/**
+ * Connect two material expression nodes
+ */
+class FECACommand_ConnectMaterialNodes : public IECACommand
+{
+public:
+	virtual FString GetName() const override { return TEXT("connect_material_nodes"); }
+	virtual FString GetDescription() const override { return TEXT("Connect two material expression nodes by their outputs/inputs"); }
+	virtual FString GetCategory() const override { return TEXT("Material Node"); }
+	
+	virtual TArray<FECACommandParam> GetParameters() const override
+	{
+		return {
+			{ TEXT("material_path"), TEXT("string"), TEXT("Path to the Material asset"), true },
+			{ TEXT("source_node_id"), TEXT("string"), TEXT("GUID of the source expression node"), true },
+			{ TEXT("source_output"), TEXT("string"), TEXT("Output name or index (RGB, R, G, B, A, or 0-4)"), false, TEXT("0") },
+			{ TEXT("target_node_id"), TEXT("string"), TEXT("GUID of the target expression node, or 'material' for material inputs"), true },
+			{ TEXT("target_input"), TEXT("string"), TEXT("Input name (BaseColor, Metallic, Roughness, Normal, EmissiveColor, etc. for material, or input index/name for nodes)"), true }
+		};
+	}
+	
+	virtual FECACommandResult Execute(const TSharedPtr<FJsonObject>& Params) override;
+};
+
+/**
+ * Get information about nodes in a material
+ */
+class FECACommand_GetMaterialNodes : public IECACommand
+{
+public:
+	virtual FString GetName() const override { return TEXT("get_material_nodes"); }
+	virtual FString GetDescription() const override { return TEXT("Get all expression nodes in a material with their connections"); }
+	virtual FString GetCategory() const override { return TEXT("Material Node"); }
+	
+	virtual TArray<FECACommandParam> GetParameters() const override
+	{
+		return {
+			{ TEXT("material_path"), TEXT("string"), TEXT("Path to the Material asset"), true },
+			{ TEXT("node_type_filter"), TEXT("string"), TEXT("Filter by node type (e.g., TextureSample)"), false }
+		};
+	}
+	
+	virtual FECACommandResult Execute(const TSharedPtr<FJsonObject>& Params) override;
+};
+
+/**
+ * Get detailed pin information for a material node
+ */
+class FECACommand_GetMaterialNodeInfo : public IECACommand
+{
+public:
+	virtual FString GetName() const override { return TEXT("get_material_node_info"); }
+	virtual FString GetDescription() const override { return TEXT("Get detailed information about a material expression node including its inputs and outputs"); }
+	virtual FString GetCategory() const override { return TEXT("Material Node"); }
+	
+	virtual TArray<FECACommandParam> GetParameters() const override
+	{
+		return {
+			{ TEXT("material_path"), TEXT("string"), TEXT("Path to the Material asset"), true },
+			{ TEXT("node_id"), TEXT("string"), TEXT("GUID of the node"), true }
+		};
+	}
+	
+	virtual FECACommandResult Execute(const TSharedPtr<FJsonObject>& Params) override;
+};
+
+/**
+ * Delete a material expression node
+ */
+class FECACommand_DeleteMaterialNode : public IECACommand
+{
+public:
+	virtual FString GetName() const override { return TEXT("delete_material_node"); }
+	virtual FString GetDescription() const override { return TEXT("Delete a material expression node by its GUID"); }
+	virtual FString GetCategory() const override { return TEXT("Material Node"); }
+	
+	virtual TArray<FECACommandParam> GetParameters() const override
+	{
+		return {
+			{ TEXT("material_path"), TEXT("string"), TEXT("Path to the Material asset"), true },
+			{ TEXT("node_id"), TEXT("string"), TEXT("GUID of the node to delete"), true }
+		};
+	}
+	
+	virtual FECACommandResult Execute(const TSharedPtr<FJsonObject>& Params) override;
+};
+
+/**
+ * Disconnect a material node's inputs or outputs
+ */
+class FECACommand_DisconnectMaterialNode : public IECACommand
+{
+public:
+	virtual FString GetName() const override { return TEXT("disconnect_material_node"); }
+	virtual FString GetDescription() const override { return TEXT("Disconnect all or specific connections on a material expression node"); }
+	virtual FString GetCategory() const override { return TEXT("Material Node"); }
+	
+	virtual TArray<FECACommandParam> GetParameters() const override
+	{
+		return {
+			{ TEXT("material_path"), TEXT("string"), TEXT("Path to the Material asset"), true },
+			{ TEXT("node_id"), TEXT("string"), TEXT("GUID of the node"), true },
+			{ TEXT("input_name"), TEXT("string"), TEXT("Specific input to disconnect (optional, disconnects all if not specified)"), false }
+		};
+	}
+	
+	virtual FECACommandResult Execute(const TSharedPtr<FJsonObject>& Params) override;
+};
+
+/**
+ * Set a material node's property value
+ */
+class FECACommand_SetMaterialNodeProperty : public IECACommand
+{
+public:
+	virtual FString GetName() const override { return TEXT("set_material_node_property"); }
+	virtual FString GetDescription() const override { return TEXT("Set a property value on a material expression node"); }
+	virtual FString GetCategory() const override { return TEXT("Material Node"); }
+	
+	virtual TArray<FECACommandParam> GetParameters() const override
+	{
+		return {
+			{ TEXT("material_path"), TEXT("string"), TEXT("Path to the Material asset"), true },
+			{ TEXT("node_id"), TEXT("string"), TEXT("GUID of the node"), true },
+			{ TEXT("property_name"), TEXT("string"), TEXT("Name of the property to set"), true },
+			{ TEXT("value"), TEXT("any"), TEXT("Value to set"), true }
+		};
+	}
+	
+	virtual FECACommandResult Execute(const TSharedPtr<FJsonObject>& Params) override;
+};
+
+/**
+ * Batch create nodes and connections in a material
+ */
+class FECACommand_BatchEditMaterialNodes : public IECACommand
+{
+public:
+	virtual FString GetName() const override { return TEXT("batch_edit_material_nodes"); }
+	virtual FString GetDescription() const override { return TEXT("Create multiple material nodes and wire them together in a single atomic operation"); }
+	virtual FString GetCategory() const override { return TEXT("Material Node"); }
+	
+	virtual TArray<FECACommandParam> GetParameters() const override
+	{
+		return {
+			{ TEXT("material_path"), TEXT("string"), TEXT("Path to the Material asset"), true },
+			{ TEXT("nodes"), TEXT("array"), TEXT("Array of node definitions. Each node must have 'temp_id' and 'node_type'. Optional: parameter_name, default_value, node_position {x, y}"), true },
+			{ TEXT("connections"), TEXT("array"), TEXT("Array of connection definitions: {source_node, source_output, target_node, target_input}. Use temp_id or existing node GUIDs. target_node can be 'material' to connect to material inputs."), false }
+		};
+	}
+	
+	virtual FECACommandResult Execute(const TSharedPtr<FJsonObject>& Params) override;
+};
+
+/**
+ * List available material expression types
+ */
+class FECACommand_ListMaterialExpressionTypes : public IECACommand
+{
+public:
+	virtual FString GetName() const override { return TEXT("list_material_expression_types"); }
+	virtual FString GetDescription() const override { return TEXT("List all available material expression types that can be added to materials"); }
+	virtual FString GetCategory() const override { return TEXT("Material Node"); }
+	
+	virtual TArray<FECACommandParam> GetParameters() const override
+	{
+		return {
+			{ TEXT("category_filter"), TEXT("string"), TEXT("Filter by category (e.g., Math, Texture, Parameters, Utility)"), false },
+			{ TEXT("search"), TEXT("string"), TEXT("Search filter for expression names"), false }
+		};
+	}
+	
+	virtual FECACommandResult Execute(const TSharedPtr<FJsonObject>& Params) override;
+};
+
+/**
+ * Get material compilation errors
+ */
+class FECACommand_GetMaterialErrors : public IECACommand
+{
+public:
+	virtual FString GetName() const override { return TEXT("get_material_errors"); }
+	virtual FString GetDescription() const override { return TEXT("Get compilation errors for a material. Call this after making changes to check for issues."); }
+	virtual FString GetCategory() const override { return TEXT("Material Node"); }
+	
+	virtual TArray<FECACommandParam> GetParameters() const override
+	{
+		return {
+			{ TEXT("material_path"), TEXT("string"), TEXT("Path to the Material asset"), true }
+		};
+	}
+	
+	virtual FECACommandResult Execute(const TSharedPtr<FJsonObject>& Params) override;
+};
+
+/**
+ * Set an input name on a Custom material expression node
+ */
+class FECACommand_SetCustomNodeInputName : public IECACommand
+{
+public:
+	virtual FString GetName() const override { return TEXT("set_custom_node_input_name"); }
+	virtual FString GetDescription() const override { return TEXT("Set the name of an input on a Custom material expression node. Creates the input if it doesn't exist at the specified index."); }
+	virtual FString GetCategory() const override { return TEXT("Material Node"); }
+	
+	virtual TArray<FECACommandParam> GetParameters() const override
+	{
+		return {
+			{ TEXT("material_path"), TEXT("string"), TEXT("Path to the Material asset"), true },
+			{ TEXT("node_id"), TEXT("string"), TEXT("GUID of the Custom node"), true },
+			{ TEXT("input_index"), TEXT("number"), TEXT("Index of the input to set (0-based)"), true },
+			{ TEXT("input_name"), TEXT("string"), TEXT("Name for the input"), true }
+		};
+	}
+	
+	virtual FECACommandResult Execute(const TSharedPtr<FJsonObject>& Params) override;
+};
+
+/**
+ * Set material properties using reflection (supports any UMaterial UPROPERTY)
+ */
+class FECACommand_SetMaterialProperty : public IECACommand
+{
+public:
+	virtual FString GetName() const override { return TEXT("set_material_property"); }
+	virtual FString GetDescription() const override { return TEXT("Set any UMaterial property using reflection. Supports snake_case or PascalCase names. Use get_material_property to discover available properties."); }
+	virtual FString GetCategory() const override { return TEXT("Material"); }
+	
+	virtual TArray<FECACommandParam> GetParameters() const override
+	{
+		return {
+			{ TEXT("material_path"), TEXT("string"), TEXT("Path to the Material asset"), true },
+			{ TEXT("property"), TEXT("string"), TEXT("Property name (e.g., BlendMode, bEnableTessellation, TwoSided, bUsedWithNanite, OpacityMaskClipValue). Supports snake_case aliases."), true },
+			{ TEXT("value"), TEXT("any"), TEXT("Value to set. For enums use the enum value name (e.g., 'BLEND_Masked'). For bools use true/false."), true }
+		};
+	}
+	
+	virtual FECACommandResult Execute(const TSharedPtr<FJsonObject>& Params) override;
+};
+
+/**
+ * Get material properties using reflection (supports any UMaterial UPROPERTY)
+ */
+class FECACommand_GetMaterialProperty : public IECACommand
+{
+public:
+	virtual FString GetName() const override { return TEXT("get_material_property"); }
+	virtual FString GetDescription() const override { return TEXT("Get UMaterial properties using reflection. Returns all editable properties if no specific property is requested."); }
+	virtual FString GetCategory() const override { return TEXT("Material"); }
+	
+	virtual TArray<FECACommandParam> GetParameters() const override
+	{
+		return {
+			{ TEXT("material_path"), TEXT("string"), TEXT("Path to the Material asset"), true },
+			{ TEXT("property"), TEXT("string"), TEXT("Specific property name to get (optional - returns all editable properties if not specified)"), false }
+		};
+	}
+	
+	virtual FECACommandResult Execute(const TSharedPtr<FJsonObject>& Params) override;
+};
+
+
+/**
+ * Auto-layout material nodes for visual clarity
+ * 
+ * Layout strategies:
+ * - horizontal: Arrange nodes left-to-right from material inputs (default)
+ * - vertical: Arrange nodes top-to-bottom from material inputs
+ * - tree: Arrange as a tree structure from material input connections
+ * - compact: Minimize total graph area while maintaining readability
+ * 
+ * The algorithm:
+ * 1. Identifies root nodes (material input connections: BaseColor, Normal, etc.)
+ * 2. Traverses expression connections backward to determine node order/depth
+ * 3. Positions nodes with consistent spacing to avoid overlaps
+ * 4. Aligns data-flow nodes near their consumers
+ */
+class FECACommand_AutoLayoutMaterialGraph : public IECACommand
+{
+public:
+	virtual FString GetName() const override { return TEXT("auto_layout_material_graph"); }
+	virtual FString GetDescription() const override { return TEXT("Automatically arrange material expression nodes for visual clarity. Organizes by connection flow with consistent spacing."); }
+	virtual FString GetCategory() const override { return TEXT("Material Node"); }
+	
+	virtual TArray<FECACommandParam> GetParameters() const override
+	{
+		return {
+			{ TEXT("material_path"), TEXT("string"), TEXT("Path to the Material asset"), true },
+			{ TEXT("strategy"), TEXT("string"), TEXT("Layout strategy: horizontal (default), vertical, tree, compact"), false, TEXT("horizontal") },
+			{ TEXT("spacing_x"), TEXT("number"), TEXT("Horizontal spacing between nodes (default: 300)"), false, TEXT("300") },
+			{ TEXT("spacing_y"), TEXT("number"), TEXT("Vertical spacing between nodes (default: 150)"), false, TEXT("150") },
+			{ TEXT("node_ids"), TEXT("array"), TEXT("Specific node IDs to layout (optional, layouts all if not specified)"), false }
+		};
+	}
+	
+	virtual FECACommandResult Execute(const TSharedPtr<FJsonObject>& Params) override;
+};
+
+/**
+ * Rename parameter groups in a material
+ * 
+ * This command renames all material parameter expressions that belong to a specific group.
+ * Useful for organizing parameters in material instances by renaming groups in bulk.
+ * 
+ * Examples:
+ * - Rename "None" group to "Base Properties" 
+ * - Rename "Textures" to "Surface Maps"
+ * - Use empty string "" for the default/unnamed group (None)
+ */
+class FECACommand_RenameParameterGroup : public IECACommand
+{
+public:
+	virtual FString GetName() const override { return TEXT("rename_parameter_group"); }
+	virtual FString GetDescription() const override { return TEXT("Rename a parameter group in a material. All parameters belonging to the old group will be moved to the new group."); }
+	virtual FString GetCategory() const override { return TEXT("Material"); }
+	
+	virtual TArray<FECACommandParam> GetParameters() const override
+	{
+		return {
+			{ TEXT("material_path"), TEXT("string"), TEXT("Path to the Material asset"), true },
+			{ TEXT("old_group_name"), TEXT("string"), TEXT("Current group name to rename (use empty string '' for the default/None group)"), true },
+			{ TEXT("new_group_name"), TEXT("string"), TEXT("New group name (use empty string '' to move to default/None group)"), true }
+		};
+	}
+	
+	virtual FECACommandResult Execute(const TSharedPtr<FJsonObject>& Params) override;
+};
+
+/**
+ * List all parameter groups in a material
+ * 
+ * Returns all unique group names used by parameters in the material,
+ * along with a count of parameters in each group.
+ */
+class FECACommand_ListParameterGroups : public IECACommand
+{
+public:
+	virtual FString GetName() const override { return TEXT("list_parameter_groups"); }
+	virtual FString GetDescription() const override { return TEXT("List all parameter groups in a material with counts of parameters in each group."); }
+	virtual FString GetCategory() const override { return TEXT("Material"); }
+	
+	virtual TArray<FECACommandParam> GetParameters() const override
+	{
+		return {
+			{ TEXT("material_path"), TEXT("string"), TEXT("Path to the Material asset"), true }
+		};
+	}
+	
+	virtual FECACommandResult Execute(const TSharedPtr<FJsonObject>& Params) override;
+};
