ECABridge experimental MCP plugin

diff --git a/Engine/Plugins/Experimental/ECABridge/Source/ECABridge/Public/Commands/ECAMeshCommands.h b/Engine/Plugins/Experimental/ECABridge/Source/ECABridge/Public/Commands/ECAMeshCommands.h
new file mode 100644
--- /dev/null
+++ b/Engine/Plugins/Experimental/ECABridge/Source/ECABridge/Public/Commands/ECAMeshCommands.h
@@ -0,0 +1,1157 @@
+// Copyright Epic Games, Inc. All Rights Reserved.
+
+#pragma once
+
+#include "Commands/ECACommand.h"
+
+/**
+ * Comprehensive procedural mesh system for ECABridge.
+ * On par with professional tools like Blender, Houdini, and OpenSCAD.
+ * 
+ * Features:
+ * - Primitive generation (box, sphere, cylinder, cone, torus, etc.)
+ * - CSG Boolean operations (union, difference, intersection)
+ * - Mesh manipulation (extrude, bevel, inset, offset)
+ * - Mesh modification (simplify, subdivide, remesh, smooth)
+ * - UV operations (auto-UV, planar projection, box projection)
+ * - Import/Export (OBJ, FBX)
+ * - Runtime procedural meshes (ProceduralMeshComponent)
+ */
+
+//==============================================================================
+// PRIMITIVE CREATION
+//==============================================================================
+
+/**
+ * Create primitive shapes (box, sphere, cylinder, cone, capsule, torus, etc.)
+ */
+class FECACommand_CreatePrimitive : public IECACommand
+{
+public:
+	virtual FString GetName() const override { return TEXT("create_primitive"); }
+	virtual FString GetDescription() const override { return TEXT("Create a primitive shape mesh (box, sphere, cylinder, cone, capsule, torus, plane, disc, stairs, arrow)"); }
+	virtual FString GetCategory() const override { return TEXT("Mesh"); }
+	
+	virtual TArray<FECACommandParam> GetParameters() const override
+	{
+		return {
+			{ TEXT("asset_path"), TEXT("string"), TEXT("Path for the new mesh asset"), true },
+			{ TEXT("type"), TEXT("string"), TEXT("Primitive type: box, sphere, cylinder, cone, capsule, torus, plane, disc, stairs, arrow, grid"), true },
+			{ TEXT("dimensions"), TEXT("object"), TEXT("Dimensions {x, y, z} or {radius, height} depending on type"), false },
+			{ TEXT("segments"), TEXT("object"), TEXT("Segment counts {radial, height, rings} for curved surfaces"), false },
+			{ TEXT("origin"), TEXT("string"), TEXT("Origin mode: center, base, corner"), false, TEXT("center") },
+			{ TEXT("material_path"), TEXT("string"), TEXT("Path to material to assign"), false },
+			{ TEXT("generate_collision"), TEXT("boolean"), TEXT("Generate collision mesh"), false, TEXT("true") },
+		};
+	}
+	
+	virtual FECACommandResult Execute(const TSharedPtr<FJsonObject>& Params) override;
+};
+
+/**
+ * Create a static mesh asset from vertex/face data
+ * 
+ * Supports multiple material slots. Use material_slots to define slot names,
+ * then use material_index in triangle objects to assign triangles to slots.
+ * 
+ * Example with multiple materials:
+ * {
+ *   "asset_path": "/Game/Meshes/MyMesh",
+ *   "vertices": [{"x":0,"y":0,"z":0}, {"x":100,"y":0,"z":0}, {"x":0,"y":100,"z":0}, {"x":100,"y":100,"z":0}],
+ *   "triangles": [
+ *     {"indices": [0,1,2], "material_index": 0},
+ *     {"indices": [1,3,2], "material_index": 1}
+ *   ],
+ *   "material_slots": ["FloorMaterial", "WallMaterial"]
+ * }
+ */
+class FECACommand_CreateMesh : public IECACommand
+{
+public:
+	virtual FString GetName() const override { return TEXT("create_mesh"); }
+	virtual FString GetDescription() const override { return TEXT("Create a static mesh asset from vertex, face, UV, and normal data. Supports multiple material slots. Example: {\"asset_path\":\"/Game/Mesh\", \"vertices\":[{\"x\":0,\"y\":0,\"z\":0},...], \"triangles\":[{\"indices\":[0,1,2], \"material_index\":0}, {\"indices\":[1,3,2], \"material_index\":1}], \"material_slots\":[\"Floor\",\"Wall\"]}"); }
+	virtual FString GetCategory() const override { return TEXT("Mesh"); }
+	
+	virtual TArray<FECACommandParam> GetParameters() const override
+	{
+		return {
+			{ TEXT("asset_path"), TEXT("string"), TEXT("Path for the new mesh asset"), true },
+			{ TEXT("vertices"), TEXT("array"), TEXT("Array of vertex positions [{x, y, z}, ...]"), true },
+			{ TEXT("triangles"), TEXT("array"), TEXT("Array of triangle indices [[v0, v1, v2], ...]"), true },
+			{ TEXT("uvs"), TEXT("array"), TEXT("Array of UV coordinates [{u, v}, ...] (per vertex)"), false },
+			{ TEXT("normals"), TEXT("array"), TEXT("Array of normals [{x, y, z}, ...] (per vertex, auto-computed if omitted)"), false },
+			{ TEXT("vertex_colors"), TEXT("array"), TEXT("Array of vertex colors [{r, g, b, a}, ...]"), false },
+			{ TEXT("material_slots"), TEXT("array"), TEXT("Material slot names and assignments"), false },
+			{ TEXT("generate_collision"), TEXT("boolean"), TEXT("Generate collision mesh"), false, TEXT("true") },
+			{ TEXT("generate_lightmap_uvs"), TEXT("boolean"), TEXT("Generate lightmap UVs"), false, TEXT("true") },
+		};
+	}
+	
+	virtual FECACommandResult Execute(const TSharedPtr<FJsonObject>& Params) override;
+};
+
+/**
+ * Create mesh from face definitions - high-level LLM-friendly interface
+ * 
+ * Instead of specifying individual vertices and triangles, define faces with
+ * their vertices directly. Supports automatic UV generation, extrusion for
+ * thickness, and per-face material assignment.
+ * 
+ * Example - Simple box walls:
+ * {
+ *   "asset_path": "/Game/Meshes/Walls",
+ *   "faces": [
+ *     {"name": "front", "vertices": [[0,0,0], [400,0,0], [400,0,250], [0,0,250]]},
+ *     {"name": "back", "vertices": [[0,300,0], [400,300,0], [400,300,250], [0,300,250]]},
+ *     {"name": "left", "vertices": [[0,0,0], [0,300,0], [0,300,250], [0,0,250]]},
+ *     {"name": "right", "vertices": [[400,0,0], [400,300,0], [400,300,250], [400,0,250]]}
+ *   ],
+ *   "thickness": 20
+ * }
+ */
+class FECACommand_CreateMeshFromFaces : public IECACommand
+{
+public:
+	virtual FString GetName() const override { return TEXT("create_mesh_from_faces"); }
+	virtual FString GetDescription() const override { return TEXT("Create mesh from face definitions with automatic triangulation and UVs. IMPORTANT: Each face MUST include 'normal_hint' to specify face direction (e.g. [0,0,1] for up, [0,1,0] for +Y). Example face: {\"name\":\"floor\", \"vertices\":[[0,0,0],[100,0,0],[100,100,0],[0,100,0]], \"normal_hint\":[0,0,1]}"); }
+	virtual FString GetCategory() const override { return TEXT("Mesh"); }
+	
+	virtual TArray<FECACommandParam> GetParameters() const override
+	{
+		return {
+			{ TEXT("asset_path"), TEXT("string"), TEXT("Path for the new mesh asset"), true },
+			{ TEXT("faces"), TEXT("array"), TEXT("Array of face objects [{name, vertices:[[x,y,z],...], material_index, normal_hint:[x,y,z]}, ...]. normal_hint auto-corrects winding."), true },
+			{ TEXT("thickness"), TEXT("number"), TEXT("Extrude faces to create solid with this thickness (0 for flat faces)"), false, TEXT("0") },
+			{ TEXT("uv_scale"), TEXT("number"), TEXT("UV scale factor (default 0.01 = 100 units per UV tile)"), false, TEXT("0.01") },
+			{ TEXT("double_sided"), TEXT("boolean"), TEXT("Generate back faces for double-sided rendering"), false, TEXT("false") },
+			{ TEXT("generate_collision"), TEXT("boolean"), TEXT("Generate collision mesh"), false, TEXT("true") },
+		};
+	}
+	
+	virtual FECACommandResult Execute(const TSharedPtr<FJsonObject>& Params) override;
+};
+
+/**
+ * Create mesh from 2D polygon points
+ */
+class FECACommand_CreatePolygonMesh : public IECACommand
+{
+public:
+	virtual FString GetName() const override { return TEXT("create_polygon_mesh"); }
+	virtual FString GetDescription() const override { return TEXT("Create a 2D polygon mesh from a list of points (can be extruded)"); }
+	virtual FString GetCategory() const override { return TEXT("Mesh"); }
+	
+	virtual TArray<FECACommandParam> GetParameters() const override
+	{
+		return {
+			{ TEXT("asset_path"), TEXT("string"), TEXT("Path for the new mesh asset"), true },
+			{ TEXT("points"), TEXT("array"), TEXT("2D polygon points [{x, y}, ...] in order"), true },
+			{ TEXT("holes"), TEXT("array"), TEXT("Optional hole polygons [[[{x,y},...]], ...]"), false },
+			{ TEXT("height"), TEXT("number"), TEXT("Extrusion height (0 for flat 2D mesh)"), false, TEXT("0") },
+			{ TEXT("material_path"), TEXT("string"), TEXT("Path to material to assign"), false },
+		};
+	}
+	
+	virtual FECACommandResult Execute(const TSharedPtr<FJsonObject>& Params) override;
+};
+
+//==============================================================================
+// CSG BOOLEAN OPERATIONS
+//==============================================================================
+
+/**
+ * Boolean/CSG operations on meshes (union, difference, intersection)
+ */
+class FECACommand_MeshBoolean : public IECACommand
+{
+public:
+	virtual FString GetName() const override { return TEXT("mesh_boolean"); }
+	virtual FString GetDescription() const override { return TEXT("Perform CSG boolean operation on two meshes"); }
+	virtual FString GetCategory() const override { return TEXT("Mesh"); }
+	
+	virtual TArray<FECACommandParam> GetParameters() const override
+	{
+		return {
+			{ TEXT("output_path"), TEXT("string"), TEXT("Path for the result mesh asset"), true },
+			{ TEXT("mesh_a"), TEXT("string"), TEXT("Path to first mesh asset"), true },
+			{ TEXT("mesh_b"), TEXT("string"), TEXT("Path to second mesh asset"), true },
+			{ TEXT("operation"), TEXT("string"), TEXT("Boolean operation: union, difference, intersection, trim_inside, trim_outside"), true },
+			{ TEXT("transform_b"), TEXT("object"), TEXT("Transform for mesh B {location:{x,y,z}, rotation:{pitch,yaw,roll}, scale:{x,y,z}}"), false },
+			{ TEXT("fill_holes"), TEXT("boolean"), TEXT("Fill holes created by the operation"), false, TEXT("true") },
+			{ TEXT("simplify"), TEXT("boolean"), TEXT("Simplify coplanar faces after operation"), false, TEXT("true") },
+		};
+	}
+	
+	virtual FECACommandResult Execute(const TSharedPtr<FJsonObject>& Params) override;
+};
+
+/**
+ * Self-union to clean up self-intersecting meshes
+ */
+class FECACommand_MeshSelfUnion : public IECACommand
+{
+public:
+	virtual FString GetName() const override { return TEXT("mesh_self_union"); }
+	virtual FString GetDescription() const override { return TEXT("Clean up self-intersecting mesh using self-union operation"); }
+	virtual FString GetCategory() const override { return TEXT("Mesh"); }
+	
+	virtual TArray<FECACommandParam> GetParameters() const override
+	{
+		return {
+			{ TEXT("mesh_path"), TEXT("string"), TEXT("Path to the mesh asset to clean"), true },
+			{ TEXT("output_path"), TEXT("string"), TEXT("Path for result (or empty to modify in place)"), false },
+			{ TEXT("fill_holes"), TEXT("boolean"), TEXT("Fill holes"), false, TEXT("true") },
+			{ TEXT("trim_flaps"), TEXT("boolean"), TEXT("Trim flaps"), false, TEXT("true") },
+		};
+	}
+	
+	virtual FECACommandResult Execute(const TSharedPtr<FJsonObject>& Params) override;
+};
+
+//==============================================================================
+// MESH CUTTING & SLICING
+//==============================================================================
+
+/**
+ * Cut mesh with a plane
+ */
+class FECACommand_MeshPlaneCut : public IECACommand
+{
+public:
+	virtual FString GetName() const override { return TEXT("mesh_plane_cut"); }
+	virtual FString GetDescription() const override { return TEXT("Cut a mesh with a plane, optionally filling the cut hole"); }
+	virtual FString GetCategory() const override { return TEXT("Mesh"); }
+	
+	virtual TArray<FECACommandParam> GetParameters() const override
+	{
+		return {
+			{ TEXT("mesh_path"), TEXT("string"), TEXT("Path to the mesh asset to cut"), true },
+			{ TEXT("output_path"), TEXT("string"), TEXT("Path for result (or empty to modify in place)"), false },
+			{ TEXT("plane_origin"), TEXT("object"), TEXT("Plane origin point {x, y, z}"), true },
+			{ TEXT("plane_normal"), TEXT("object"), TEXT("Plane normal direction {x, y, z}"), true },
+			{ TEXT("fill_hole"), TEXT("boolean"), TEXT("Fill the cut hole"), false, TEXT("true") },
+			{ TEXT("flip_side"), TEXT("boolean"), TEXT("Keep opposite side of cut"), false, TEXT("false") },
+		};
+	}
+	
+	virtual FECACommandResult Execute(const TSharedPtr<FJsonObject>& Params) override;
+};
+
+/**
+ * Slice mesh into two pieces
+ */
+class FECACommand_MeshSlice : public IECACommand
+{
+public:
+	virtual FString GetName() const override { return TEXT("mesh_slice"); }
+	virtual FString GetDescription() const override { return TEXT("Slice a mesh into two separate pieces"); }
+	virtual FString GetCategory() const override { return TEXT("Mesh"); }
+	
+	virtual TArray<FECACommandParam> GetParameters() const override
+	{
+		return {
+			{ TEXT("mesh_path"), TEXT("string"), TEXT("Path to the mesh asset to slice"), true },
+			{ TEXT("output_path_a"), TEXT("string"), TEXT("Path for first half"), true },
+			{ TEXT("output_path_b"), TEXT("string"), TEXT("Path for second half"), true },
+			{ TEXT("plane_origin"), TEXT("object"), TEXT("Plane origin point {x, y, z}"), true },
+			{ TEXT("plane_normal"), TEXT("object"), TEXT("Plane normal direction {x, y, z}"), true },
+			{ TEXT("fill_holes"), TEXT("boolean"), TEXT("Fill cut holes on both pieces"), false, TEXT("true") },
+			{ TEXT("gap_width"), TEXT("number"), TEXT("Gap width between slices"), false, TEXT("0") },
+		};
+	}
+	
+	virtual FECACommandResult Execute(const TSharedPtr<FJsonObject>& Params) override;
+};
+
+/**
+ * Mirror mesh across a plane
+ */
+class FECACommand_MeshMirror : public IECACommand
+{
+public:
+	virtual FString GetName() const override { return TEXT("mesh_mirror"); }
+	virtual FString GetDescription() const override { return TEXT("Mirror a mesh across a plane"); }
+	virtual FString GetCategory() const override { return TEXT("Mesh"); }
+	
+	virtual TArray<FECACommandParam> GetParameters() const override
+	{
+		return {
+			{ TEXT("mesh_path"), TEXT("string"), TEXT("Path to the mesh asset to mirror"), true },
+			{ TEXT("output_path"), TEXT("string"), TEXT("Path for result (or empty to modify in place)"), false },
+			{ TEXT("mirror_axis"), TEXT("string"), TEXT("Mirror axis: x, y, z (or provide plane_origin/plane_normal)"), false, TEXT("x") },
+			{ TEXT("plane_origin"), TEXT("object"), TEXT("Custom plane origin {x, y, z}"), false },
+			{ TEXT("plane_normal"), TEXT("object"), TEXT("Custom plane normal {x, y, z}"), false },
+			{ TEXT("weld"), TEXT("boolean"), TEXT("Weld vertices along mirror plane"), false, TEXT("true") },
+			{ TEXT("cut"), TEXT("boolean"), TEXT("Cut mesh at mirror plane first"), false, TEXT("true") },
+		};
+	}
+	
+	virtual FECACommandResult Execute(const TSharedPtr<FJsonObject>& Params) override;
+};
+
+//==============================================================================
+// MESH MODELING OPERATIONS
+//==============================================================================
+
+/**
+ * Extrude faces/polygon
+ */
+class FECACommand_MeshExtrude : public IECACommand
+{
+public:
+	virtual FString GetName() const override { return TEXT("mesh_extrude"); }
+	virtual FString GetDescription() const override { return TEXT("Extrude selected faces or entire mesh"); }
+	virtual FString GetCategory() const override { return TEXT("Mesh"); }
+	
+	virtual TArray<FECACommandParam> GetParameters() const override
+	{
+		return {
+			{ TEXT("mesh_path"), TEXT("string"), TEXT("Path to the mesh asset"), true },
+			{ TEXT("output_path"), TEXT("string"), TEXT("Path for result (or empty to modify in place)"), false },
+			{ TEXT("distance"), TEXT("number"), TEXT("Extrusion distance"), true },
+			{ TEXT("direction"), TEXT("object"), TEXT("Extrusion direction {x,y,z} (or empty for face normal)"), false },
+			{ TEXT("selection"), TEXT("object"), TEXT("Face selection {polygroup: name} or {triangles: [0,1,2,...]}"), false },
+			{ TEXT("create_shells"), TEXT("boolean"), TEXT("Convert solid to shell during extrusion"), false, TEXT("true") },
+		};
+	}
+	
+	virtual FECACommandResult Execute(const TSharedPtr<FJsonObject>& Params) override;
+};
+
+/**
+ * Linear extrude 2D profile to 3D
+ */
+class FECACommand_ExtrudePolygon : public IECACommand
+{
+public:
+	virtual FString GetName() const override { return TEXT("extrude_polygon"); }
+	virtual FString GetDescription() const override { return TEXT("Extrude a 2D polygon profile into a 3D mesh"); }
+	virtual FString GetCategory() const override { return TEXT("Mesh"); }
+	
+	virtual TArray<FECACommandParam> GetParameters() const override
+	{
+		return {
+			{ TEXT("asset_path"), TEXT("string"), TEXT("Path for the new mesh asset"), true },
+			{ TEXT("points"), TEXT("array"), TEXT("2D profile points [{x, y}, ...] in order"), true },
+			{ TEXT("height"), TEXT("number"), TEXT("Extrusion height"), true },
+			{ TEXT("twist"), TEXT("number"), TEXT("Twist angle in degrees during extrusion"), false, TEXT("0") },
+			{ TEXT("scale_end"), TEXT("number"), TEXT("Scale factor at the end of extrusion"), false, TEXT("1") },
+			{ TEXT("segments"), TEXT("number"), TEXT("Number of segments along extrusion"), false, TEXT("1") },
+			{ TEXT("cap_start"), TEXT("boolean"), TEXT("Cap the start of extrusion"), false, TEXT("true") },
+			{ TEXT("cap_end"), TEXT("boolean"), TEXT("Cap the end of extrusion"), false, TEXT("true") },
+		};
+	}
+	
+	virtual FECACommandResult Execute(const TSharedPtr<FJsonObject>& Params) override;
+};
+
+/**
+ * Revolve 2D profile around an axis (lathe)
+ */
+class FECACommand_RevolveProfile : public IECACommand
+{
+public:
+	virtual FString GetName() const override { return TEXT("revolve_profile"); }
+	virtual FString GetDescription() const override { return TEXT("Revolve a 2D profile around an axis (lathe operation)"); }
+	virtual FString GetCategory() const override { return TEXT("Mesh"); }
+	
+	virtual TArray<FECACommandParam> GetParameters() const override
+	{
+		return {
+			{ TEXT("asset_path"), TEXT("string"), TEXT("Path for the new mesh asset"), true },
+			{ TEXT("points"), TEXT("array"), TEXT("2D profile points [{x, y}, ...] where X is radius, Y is height"), true },
+			{ TEXT("angle"), TEXT("number"), TEXT("Revolve angle in degrees (360 for full revolution)"), false, TEXT("360") },
+			{ TEXT("segments"), TEXT("number"), TEXT("Number of radial segments"), false, TEXT("32") },
+			{ TEXT("axis"), TEXT("string"), TEXT("Axis to revolve around: x, y, z"), false, TEXT("z") },
+			{ TEXT("cap_start"), TEXT("boolean"), TEXT("Cap start if angle < 360"), false, TEXT("true") },
+			{ TEXT("cap_end"), TEXT("boolean"), TEXT("Cap end if angle < 360"), false, TEXT("true") },
+			{ TEXT("hard_normals"), TEXT("boolean"), TEXT("Use hard normals at profile corners"), false, TEXT("false") },
+		};
+	}
+	
+	virtual FECACommandResult Execute(const TSharedPtr<FJsonObject>& Params) override;
+};
+
+/**
+ * Sweep profile along a path
+ */
+class FECACommand_SweepProfile : public IECACommand
+{
+public:
+	virtual FString GetName() const override { return TEXT("sweep_profile"); }
+	virtual FString GetDescription() const override { return TEXT("Sweep a 2D profile along a 3D path"); }
+	virtual FString GetCategory() const override { return TEXT("Mesh"); }
+	
+	virtual TArray<FECACommandParam> GetParameters() const override
+	{
+		return {
+			{ TEXT("asset_path"), TEXT("string"), TEXT("Path for the new mesh asset"), true },
+			{ TEXT("profile_points"), TEXT("array"), TEXT("2D profile points [{x, y}, ...]"), true },
+			{ TEXT("path_points"), TEXT("array"), TEXT("3D path points [{x, y, z}, ...]"), true },
+			{ TEXT("twist_total"), TEXT("number"), TEXT("Total twist angle along path"), false, TEXT("0") },
+			{ TEXT("scale_curve"), TEXT("array"), TEXT("Scale values along path [s0, s1, ...]"), false },
+			{ TEXT("closed_path"), TEXT("boolean"), TEXT("Is the path closed (loop)"), false, TEXT("false") },
+		};
+	}
+	
+	virtual FECACommandResult Execute(const TSharedPtr<FJsonObject>& Params) override;
+};
+
+/**
+ * Inset faces
+ */
+class FECACommand_MeshInset : public IECACommand
+{
+public:
+	virtual FString GetName() const override { return TEXT("mesh_inset"); }
+	virtual FString GetDescription() const override { return TEXT("Inset faces inward or outward"); }
+	virtual FString GetCategory() const override { return TEXT("Mesh"); }
+	
+	virtual TArray<FECACommandParam> GetParameters() const override
+	{
+		return {
+			{ TEXT("mesh_path"), TEXT("string"), TEXT("Path to the mesh asset"), true },
+			{ TEXT("output_path"), TEXT("string"), TEXT("Path for result (or empty to modify in place)"), false },
+			{ TEXT("distance"), TEXT("number"), TEXT("Inset distance (positive=inward, negative=outward)"), true },
+			{ TEXT("selection"), TEXT("object"), TEXT("Face selection {polygroup: name} or {triangles: [0,1,2,...]}"), false },
+		};
+	}
+	
+	virtual FECACommandResult Execute(const TSharedPtr<FJsonObject>& Params) override;
+};
+
+/**
+ * Offset mesh surface (shell/solidify)
+ */
+class FECACommand_MeshOffset : public IECACommand
+{
+public:
+	virtual FString GetName() const override { return TEXT("mesh_offset"); }
+	virtual FString GetDescription() const override { return TEXT("Offset mesh surface to create shells or thicken walls"); }
+	virtual FString GetCategory() const override { return TEXT("Mesh"); }
+	
+	virtual TArray<FECACommandParam> GetParameters() const override
+	{
+		return {
+			{ TEXT("mesh_path"), TEXT("string"), TEXT("Path to the mesh asset"), true },
+			{ TEXT("output_path"), TEXT("string"), TEXT("Path for result (or empty to modify in place)"), false },
+			{ TEXT("distance"), TEXT("number"), TEXT("Offset distance"), true },
+			{ TEXT("fixed_boundary"), TEXT("boolean"), TEXT("Keep boundary edges fixed"), false, TEXT("false") },
+			{ TEXT("smooth_steps"), TEXT("number"), TEXT("Number of smoothing iterations"), false, TEXT("5") },
+		};
+	}
+	
+	virtual FECACommandResult Execute(const TSharedPtr<FJsonObject>& Params) override;
+};
+
+//==============================================================================
+// MESH MODIFICATION (SIMPLIFY, SUBDIVIDE, REMESH, SMOOTH)
+//==============================================================================
+
+/**
+ * Simplify mesh (reduce polygon count)
+ */
+class FECACommand_MeshSimplify : public IECACommand
+{
+public:
+	virtual FString GetName() const override { return TEXT("mesh_simplify"); }
+	virtual FString GetDescription() const override { return TEXT("Reduce mesh polygon count while preserving shape"); }
+	virtual FString GetCategory() const override { return TEXT("Mesh"); }
+	
+	virtual TArray<FECACommandParam> GetParameters() const override
+	{
+		return {
+			{ TEXT("mesh_path"), TEXT("string"), TEXT("Path to the mesh asset"), true },
+			{ TEXT("output_path"), TEXT("string"), TEXT("Path for result (or empty to modify in place)"), false },
+			{ TEXT("method"), TEXT("string"), TEXT("Simplification method: triangle_count, vertex_count, percentage, tolerance"), true },
+			{ TEXT("target"), TEXT("number"), TEXT("Target value (count, percentage 0-100, or tolerance distance)"), true },
+			{ TEXT("preserve_boundaries"), TEXT("boolean"), TEXT("Preserve mesh boundaries"), false, TEXT("true") },
+			{ TEXT("preserve_uvs"), TEXT("boolean"), TEXT("Preserve UV seams"), false, TEXT("true") },
+		};
+	}
+	
+	virtual FECACommandResult Execute(const TSharedPtr<FJsonObject>& Params) override;
+};
+
+/**
+ * Subdivide mesh (increase polygon count)
+ */
+class FECACommand_MeshSubdivide : public IECACommand
+{
+public:
+	virtual FString GetName() const override { return TEXT("mesh_subdivide"); }
+	virtual FString GetDescription() const override { return TEXT("Subdivide mesh to increase polygon count and smoothness"); }
+	virtual FString GetCategory() const override { return TEXT("Mesh"); }
+	
+	virtual TArray<FECACommandParam> GetParameters() const override
+	{
+		return {
+			{ TEXT("mesh_path"), TEXT("string"), TEXT("Path to the mesh asset"), true },
+			{ TEXT("output_path"), TEXT("string"), TEXT("Path for result (or empty to modify in place)"), false },
+			{ TEXT("method"), TEXT("string"), TEXT("Subdivision method: 'uniform' (simple tessellation), 'loop' (best for triangles), 'catmull_clark' (best for quads), 'bilinear' (no smoothing), 'pn_triangles' (curved surface)"), false, TEXT("uniform") },
+			{ TEXT("level"), TEXT("number"), TEXT("Subdivision level (each level roughly quadruples triangles)"), false, TEXT("1") },
+		};
+	}
+	
+	virtual FECACommandResult Execute(const TSharedPtr<FJsonObject>& Params) override;
+};
+
+/**
+ * Remesh to uniform triangles
+ */
+class FECACommand_MeshRemesh : public IECACommand
+{
+public:
+	virtual FString GetName() const override { return TEXT("mesh_remesh"); }
+	virtual FString GetDescription() const override { return TEXT("Remesh to create uniform triangle sizes"); }
+	virtual FString GetCategory() const override { return TEXT("Mesh"); }
+	
+	virtual TArray<FECACommandParam> GetParameters() const override
+	{
+		return {
+			{ TEXT("mesh_path"), TEXT("string"), TEXT("Path to the mesh asset"), true },
+			{ TEXT("output_path"), TEXT("string"), TEXT("Path for result (or empty to modify in place)"), false },
+			{ TEXT("target_edge_length"), TEXT("number"), TEXT("Target edge length for uniform triangles"), true },
+			{ TEXT("smoothing_rate"), TEXT("number"), TEXT("Smoothing rate 0-1"), false, TEXT("0.25") },
+			{ TEXT("iterations"), TEXT("number"), TEXT("Number of remesh iterations"), false, TEXT("20") },
+			{ TEXT("preserve_boundaries"), TEXT("boolean"), TEXT("Preserve mesh boundaries"), false, TEXT("true") },
+		};
+	}
+	
+	virtual FECACommandResult Execute(const TSharedPtr<FJsonObject>& Params) override;
+};
+
+/**
+ * Smooth mesh vertices
+ */
+class FECACommand_MeshSmooth : public IECACommand
+{
+public:
+	virtual FString GetName() const override { return TEXT("mesh_smooth"); }
+	virtual FString GetDescription() const override { return TEXT("Smooth mesh vertices using various methods"); }
+	virtual FString GetCategory() const override { return TEXT("Mesh"); }
+	
+	virtual TArray<FECACommandParam> GetParameters() const override
+	{
+		return {
+			{ TEXT("mesh_path"), TEXT("string"), TEXT("Path to the mesh asset"), true },
+			{ TEXT("output_path"), TEXT("string"), TEXT("Path for result (or empty to modify in place)"), false },
+			{ TEXT("method"), TEXT("string"), TEXT("Smoothing method: laplacian, mean_value, cotangent"), false, TEXT("laplacian") },
+			{ TEXT("strength"), TEXT("number"), TEXT("Smoothing strength 0-1"), false, TEXT("0.5") },
+			{ TEXT("iterations"), TEXT("number"), TEXT("Number of smoothing iterations"), false, TEXT("1") },
+			{ TEXT("preserve_boundaries"), TEXT("boolean"), TEXT("Preserve mesh boundaries"), false, TEXT("true") },
+		};
+	}
+	
+	virtual FECACommandResult Execute(const TSharedPtr<FJsonObject>& Params) override;
+};
+
+//==============================================================================
+// MESH DEFORMATION
+//==============================================================================
+
+/**
+ * Displace mesh along normals
+ */
+class FECACommand_MeshDisplace : public IECACommand
+{
+public:
+	virtual FString GetName() const override { return TEXT("mesh_displace"); }
+	virtual FString GetDescription() const override { return TEXT("Displace mesh vertices along normals using noise or texture"); }
+	virtual FString GetCategory() const override { return TEXT("Mesh"); }
+	
+	virtual TArray<FECACommandParam> GetParameters() const override
+	{
+		return {
+			{ TEXT("mesh_path"), TEXT("string"), TEXT("Path to the mesh asset"), true },
+			{ TEXT("output_path"), TEXT("string"), TEXT("Path for result (or empty to modify in place)"), false },
+			{ TEXT("displacement"), TEXT("number"), TEXT("Displacement distance"), true },
+			{ TEXT("texture_path"), TEXT("string"), TEXT("Optional displacement texture path"), false },
+			{ TEXT("noise_scale"), TEXT("number"), TEXT("Noise scale if no texture (0 = uniform displacement)"), false, TEXT("0") },
+			{ TEXT("noise_frequency"), TEXT("number"), TEXT("Noise frequency"), false, TEXT("1") },
+		};
+	}
+	
+	virtual FECACommandResult Execute(const TSharedPtr<FJsonObject>& Params) override;
+};
+
+/**
+ * Bend mesh along an axis
+ */
+class FECACommand_MeshBend : public IECACommand
+{
+public:
+	virtual FString GetName() const override { return TEXT("mesh_bend"); }
+	virtual FString GetDescription() const override { return TEXT("Bend mesh along an axis"); }
+	virtual FString GetCategory() const override { return TEXT("Mesh"); }
+	
+	virtual TArray<FECACommandParam> GetParameters() const override
+	{
+		return {
+			{ TEXT("mesh_path"), TEXT("string"), TEXT("Path to the mesh asset"), true },
+			{ TEXT("output_path"), TEXT("string"), TEXT("Path for result (or empty to modify in place)"), false },
+			{ TEXT("angle"), TEXT("number"), TEXT("Bend angle in degrees"), true },
+			{ TEXT("axis"), TEXT("string"), TEXT("Bend axis: x, y, z"), false, TEXT("z") },
+			{ TEXT("center"), TEXT("object"), TEXT("Bend center point {x, y, z}"), false },
+		};
+	}
+	
+	virtual FECACommandResult Execute(const TSharedPtr<FJsonObject>& Params) override;
+};
+
+/**
+ * Twist mesh around an axis
+ */
+class FECACommand_MeshTwist : public IECACommand
+{
+public:
+	virtual FString GetName() const override { return TEXT("mesh_twist"); }
+	virtual FString GetDescription() const override { return TEXT("Twist mesh around an axis"); }
+	virtual FString GetCategory() const override { return TEXT("Mesh"); }
+	
+	virtual TArray<FECACommandParam> GetParameters() const override
+	{
+		return {
+			{ TEXT("mesh_path"), TEXT("string"), TEXT("Path to the mesh asset"), true },
+			{ TEXT("output_path"), TEXT("string"), TEXT("Path for result (or empty to modify in place)"), false },
+			{ TEXT("angle"), TEXT("number"), TEXT("Total twist angle in degrees"), true },
+			{ TEXT("axis"), TEXT("string"), TEXT("Twist axis: x, y, z"), false, TEXT("z") },
+			{ TEXT("symmetric"), TEXT("boolean"), TEXT("Symmetric twist around center"), false, TEXT("false") },
+		};
+	}
+	
+	virtual FECACommandResult Execute(const TSharedPtr<FJsonObject>& Params) override;
+};
+
+//==============================================================================
+// HEIGHTFIELD / TERRAIN
+//==============================================================================
+
+/**
+ * Create mesh from heightmap data
+ */
+class FECACommand_CreateHeightfield : public IECACommand
+{
+public:
+	virtual FString GetName() const override { return TEXT("create_heightfield"); }
+	virtual FString GetDescription() const override { return TEXT("Create a mesh from heightmap data or texture"); }
+	virtual FString GetCategory() const override { return TEXT("Mesh"); }
+	
+	virtual TArray<FECACommandParam> GetParameters() const override
+	{
+		return {
+			{ TEXT("asset_path"), TEXT("string"), TEXT("Path for the new mesh asset"), true },
+			{ TEXT("heights"), TEXT("array"), TEXT("2D array of height values [[row0], [row1], ...]"), false },
+			{ TEXT("texture_path"), TEXT("string"), TEXT("Path to texture to use as heightmap"), false },
+			{ TEXT("size_x"), TEXT("number"), TEXT("World size in X"), false, TEXT("1000") },
+			{ TEXT("size_y"), TEXT("number"), TEXT("World size in Y"), false, TEXT("1000") },
+			{ TEXT("height_scale"), TEXT("number"), TEXT("Height scale multiplier"), false, TEXT("100") },
+			{ TEXT("subdivisions_x"), TEXT("number"), TEXT("Grid subdivisions in X"), false, TEXT("64") },
+			{ TEXT("subdivisions_y"), TEXT("number"), TEXT("Grid subdivisions in Y"), false, TEXT("64") },
+		};
+	}
+	
+	virtual FECACommandResult Execute(const TSharedPtr<FJsonObject>& Params) override;
+};
+
+//==============================================================================
+// UV OPERATIONS
+//==============================================================================
+
+/**
+ * Generate UVs for mesh
+ */
+class FECACommand_MeshGenerateUVs : public IECACommand
+{
+public:
+	virtual FString GetName() const override { return TEXT("mesh_generate_uvs"); }
+	virtual FString GetDescription() const override { return TEXT("Generate UVs using various projection methods"); }
+	virtual FString GetCategory() const override { return TEXT("Mesh"); }
+	
+	virtual TArray<FECACommandParam> GetParameters() const override
+	{
+		return {
+			{ TEXT("mesh_path"), TEXT("string"), TEXT("Path to the mesh asset"), true },
+			{ TEXT("output_path"), TEXT("string"), TEXT("Path for result (or empty to modify in place)"), false },
+			{ TEXT("method"), TEXT("string"), TEXT("UV method: auto, planar, box, cylindrical, spherical"), false, TEXT("auto") },
+			{ TEXT("uv_channel"), TEXT("number"), TEXT("UV channel to generate (0-7)"), false, TEXT("0") },
+			{ TEXT("scale"), TEXT("number"), TEXT("UV scale"), false, TEXT("1") },
+			{ TEXT("projection_axis"), TEXT("string"), TEXT("Projection axis for planar: x, y, z"), false, TEXT("z") },
+		};
+	}
+	
+	virtual FECACommandResult Execute(const TSharedPtr<FJsonObject>& Params) override;
+};
+
+//==============================================================================
+// MESH NORMALS
+//==============================================================================
+
+/**
+ * Recompute mesh normals
+ * 
+ * By default uses seam-aware normal computation that properly averages normals
+ * across vertices that share the same position (UV seams, hard edge splits).
+ * This produces correct smooth shading across UV seams.
+ * 
+ * Set weld_seams=false to use per-vertex-index computation (faster but doesn't
+ * smooth across seams).
+ */
+class FECACommand_MeshRecomputeNormals : public IECACommand
+{
+public:
+	virtual FString GetName() const override { return TEXT("mesh_recompute_normals"); }
+	virtual FString GetDescription() const override { return TEXT("Recompute mesh normals with seam-aware smoothing. By default welds normals across UV seams for proper smooth shading."); }
+	virtual FString GetCategory() const override { return TEXT("Mesh"); }
+	
+	virtual TArray<FECACommandParam> GetParameters() const override
+	{
+		return {
+			{ TEXT("mesh_path"), TEXT("string"), TEXT("Path to the mesh asset"), true },
+			{ TEXT("output_path"), TEXT("string"), TEXT("Path for result (or empty to modify in place)"), false },
+			{ TEXT("weld_seams"), TEXT("boolean"), TEXT("Weld normals across UV seams and split vertices at the same position (recommended for smooth shading)"), false, TEXT("true") },
+			{ TEXT("weld_tolerance"), TEXT("number"), TEXT("Distance tolerance for considering vertices as colocated"), false, TEXT("0.0001") },
+			{ TEXT("weight_by_area"), TEXT("boolean"), TEXT("Weight face normals by triangle area"), false, TEXT("true") },
+			{ TEXT("weight_by_angle"), TEXT("boolean"), TEXT("Weight face normals by angle at vertex"), false, TEXT("true") },
+		};
+	}
+	
+	virtual FECACommandResult Execute(const TSharedPtr<FJsonObject>& Params) override;
+};
+
+/**
+ * Flip mesh normals
+ */
+class FECACommand_MeshFlipNormals : public IECACommand
+{
+public:
+	virtual FString GetName() const override { return TEXT("mesh_flip_normals"); }
+	virtual FString GetDescription() const override { return TEXT("Flip mesh normals (inside-out)"); }
+	virtual FString GetCategory() const override { return TEXT("Mesh"); }
+	
+	virtual TArray<FECACommandParam> GetParameters() const override
+	{
+		return {
+			{ TEXT("mesh_path"), TEXT("string"), TEXT("Path to the mesh asset"), true },
+			{ TEXT("output_path"), TEXT("string"), TEXT("Path for result (or empty to modify in place)"), false },
+		};
+	}
+	
+	virtual FECACommandResult Execute(const TSharedPtr<FJsonObject>& Params) override;
+};
+
+/**
+ * Reverse mesh triangle winding order (does not change normals)
+ */
+class FECACommand_MeshReverseWinding : public IECACommand
+{
+public:
+	virtual FString GetName() const override { return TEXT("mesh_reverse_winding"); }
+	virtual FString GetDescription() const override { return TEXT("Reverse triangle winding order only (does not change normal vectors)"); }
+	virtual FString GetCategory() const override { return TEXT("Mesh"); }
+	
+	virtual TArray<FECACommandParam> GetParameters() const override
+	{
+		return {
+			{ TEXT("mesh_path"), TEXT("string"), TEXT("Path to the mesh asset"), true },
+			{ TEXT("output_path"), TEXT("string"), TEXT("Path for result (or empty to modify in place)"), false },
+		};
+	}
+	
+	virtual FECACommandResult Execute(const TSharedPtr<FJsonObject>& Params) override;
+};
+
+//==============================================================================
+// MESH REPAIR
+//==============================================================================
+
+/**
+ * Repair mesh (fix holes, self-intersections, etc.)
+ */
+class FECACommand_MeshRepair : public IECACommand
+{
+public:
+	virtual FString GetName() const override { return TEXT("mesh_repair"); }
+	virtual FString GetDescription() const override { return TEXT("Repair mesh issues like holes, non-manifold edges, etc."); }
+	virtual FString GetCategory() const override { return TEXT("Mesh"); }
+	
+	virtual TArray<FECACommandParam> GetParameters() const override
+	{
+		return {
+			{ TEXT("mesh_path"), TEXT("string"), TEXT("Path to the mesh asset"), true },
+			{ TEXT("output_path"), TEXT("string"), TEXT("Path for result (or empty to modify in place)"), false },
+			{ TEXT("fill_holes"), TEXT("boolean"), TEXT("Fill holes"), false, TEXT("true") },
+			{ TEXT("remove_small_components"), TEXT("boolean"), TEXT("Remove small disconnected components"), false, TEXT("true") },
+			{ TEXT("min_component_volume"), TEXT("number"), TEXT("Minimum component volume to keep"), false, TEXT("0.01") },
+			{ TEXT("weld_tolerance"), TEXT("number"), TEXT("Vertex welding tolerance"), false, TEXT("0.001") },
+		};
+	}
+	
+	virtual FECACommandResult Execute(const TSharedPtr<FJsonObject>& Params) override;
+};
+
+//==============================================================================
+// MESH MATERIALS
+//==============================================================================
+
+/**
+ * Set material slots on a mesh
+ * 
+ * Can rename slots, assign materials, and optionally reassign triangles to different slots.
+ * 
+ * Example - rename slots and assign materials:
+ * {
+ *   "mesh_path": "/Game/Meshes/MyMesh",
+ *   "material_slots": [
+ *     {"name": "Wood", "material_path": "/Game/Materials/M_Wood"},
+ *     {"name": "Metal", "material_path": "/Game/Materials/M_Metal"}
+ *   ]
+ * }
+ * 
+ * Example - reassign triangles to different material slots:
+ * {
+ *   "mesh_path": "/Game/Meshes/MyMesh",
+ *   "material_slots": ["SlotA", "SlotB"],
+ *   "triangle_materials": [0, 0, 1, 1, 0, 1]
+ * }
+ */
+class FECACommand_MeshSetMaterials : public IECACommand
+{
+public:
+	virtual FString GetName() const override { return TEXT("mesh_set_materials"); }
+	virtual FString GetDescription() const override { return TEXT("Set material slots and assignments on a mesh. Example: {\"mesh_path\":\"/Game/MyMesh\", \"material_slots\":[{\"name\":\"Wood\", \"material_path\":\"/Game/M_Wood\"}, {\"name\":\"Metal\"}], \"triangle_materials\":[0,0,1,1,0,1]}"); }
+	virtual FString GetCategory() const override { return TEXT("Mesh"); }
+	
+	virtual TArray<FECACommandParam> GetParameters() const override
+	{
+		return {
+			{ TEXT("mesh_path"), TEXT("string"), TEXT("Path to the mesh asset"), true },
+			{ TEXT("material_slots"), TEXT("array"), TEXT("Material slot definitions [{name, material_path}, ...]"), true },
+			{ TEXT("triangle_materials"), TEXT("array"), TEXT("Optional per-triangle material indices [0, 1, 0, ...]"), false },
+		};
+	}
+	
+	virtual FECACommandResult Execute(const TSharedPtr<FJsonObject>& Params) override;
+};
+
+//==============================================================================
+// MESH TRANSFORM
+//==============================================================================
+
+/**
+ * Transform mesh vertices
+ */
+class FECACommand_MeshTransform : public IECACommand
+{
+public:
+	virtual FString GetName() const override { return TEXT("mesh_transform"); }
+	virtual FString GetDescription() const override { return TEXT("Apply transform to mesh vertices"); }
+	virtual FString GetCategory() const override { return TEXT("Mesh"); }
+	
+	virtual TArray<FECACommandParam> GetParameters() const override
+	{
+		return {
+			{ TEXT("mesh_path"), TEXT("string"), TEXT("Path to the mesh asset"), true },
+			{ TEXT("output_path"), TEXT("string"), TEXT("Path for result (or empty to modify in place)"), false },
+			{ TEXT("location"), TEXT("object"), TEXT("Translation {x, y, z}"), false },
+			{ TEXT("rotation"), TEXT("object"), TEXT("Rotation {pitch, yaw, roll} in degrees"), false },
+			{ TEXT("scale"), TEXT("object"), TEXT("Scale {x, y, z}"), false },
+			{ TEXT("pivot"), TEXT("object"), TEXT("Transform pivot point {x, y, z}"), false },
+		};
+	}
+	
+	virtual FECACommandResult Execute(const TSharedPtr<FJsonObject>& Params) override;
+};
+
+//==============================================================================
+// MESH QUERIES
+//==============================================================================
+
+/**
+ * Get detailed mesh information
+ */
+class FECACommand_GetMeshInfo : public IECACommand
+{
+public:
+	virtual FString GetName() const override { return TEXT("get_mesh_info"); }
+	virtual FString GetDescription() const override { return TEXT("Get detailed information about a mesh asset"); }
+	virtual FString GetCategory() const override { return TEXT("Mesh"); }
+	
+	virtual TArray<FECACommandParam> GetParameters() const override
+	{
+		return {
+			{ TEXT("mesh_path"), TEXT("string"), TEXT("Path to the mesh asset"), true },
+		};
+	}
+	
+	virtual FECACommandResult Execute(const TSharedPtr<FJsonObject>& Params) override;
+};
+
+/**
+ * Get mesh bounding box
+ */
+class FECACommand_GetMeshBounds : public IECACommand
+{
+public:
+	virtual FString GetName() const override { return TEXT("get_mesh_bounds"); }
+	virtual FString GetDescription() const override { return TEXT("Get mesh bounding box information"); }
+	virtual FString GetCategory() const override { return TEXT("Mesh"); }
+	
+	virtual TArray<FECACommandParam> GetParameters() const override
+	{
+		return {
+			{ TEXT("mesh_path"), TEXT("string"), TEXT("Path to the mesh asset"), true },
+		};
+	}
+	
+	virtual FECACommandResult Execute(const TSharedPtr<FJsonObject>& Params) override;
+};
+
+//==============================================================================
+// RUNTIME PROCEDURAL MESHES
+//==============================================================================
+
+/**
+ * Spawn a procedural mesh component (runtime, no asset)
+ */
+class FECACommand_SpawnProceduralMesh : public IECACommand
+{
+public:
+	virtual FString GetName() const override { return TEXT("spawn_procedural_mesh"); }
+	virtual FString GetDescription() const override { return TEXT("Spawn a procedural mesh actor in the world (runtime, no asset file)"); }
+	virtual FString GetCategory() const override { return TEXT("Mesh"); }
+	
+	virtual TArray<FECACommandParam> GetParameters() const override
+	{
+		return {
+			{ TEXT("actor_name"), TEXT("string"), TEXT("Name for the spawned actor"), true },
+			{ TEXT("vertices"), TEXT("array"), TEXT("Array of vertex positions [{x, y, z}, ...]"), true },
+			{ TEXT("triangles"), TEXT("array"), TEXT("Array of triangle indices [v0, v1, v2, ...] (flat array)"), true },
+			{ TEXT("location"), TEXT("object"), TEXT("Spawn location {x, y, z}"), false },
+			{ TEXT("rotation"), TEXT("object"), TEXT("Spawn rotation {pitch, yaw, roll}"), false },
+			{ TEXT("uvs"), TEXT("array"), TEXT("Array of UV coordinates [{u, v}, ...]"), false },
+			{ TEXT("normals"), TEXT("array"), TEXT("Array of normals [{x, y, z}, ...]"), false },
+			{ TEXT("vertex_colors"), TEXT("array"), TEXT("Array of vertex colors [{r, g, b, a}, ...]"), false },
+			{ TEXT("material_path"), TEXT("string"), TEXT("Path to material to assign"), false },
+			{ TEXT("enable_collision"), TEXT("boolean"), TEXT("Enable collision"), false, TEXT("true") },
+		};
+	}
+	
+	virtual FECACommandResult Execute(const TSharedPtr<FJsonObject>& Params) override;
+};
+
+/**
+ * Update procedural mesh component data
+ */
+class FECACommand_UpdateProceduralMesh : public IECACommand
+{
+public:
+	virtual FString GetName() const override { return TEXT("update_procedural_mesh"); }
+	virtual FString GetDescription() const override { return TEXT("Update an existing procedural mesh actor's geometry"); }
+	virtual FString GetCategory() const override { return TEXT("Mesh"); }
+	
+	virtual TArray<FECACommandParam> GetParameters() const override
+	{
+		return {
+			{ TEXT("actor_name"), TEXT("string"), TEXT("Name of the procedural mesh actor"), true },
+			{ TEXT("section_index"), TEXT("number"), TEXT("Mesh section index to update"), false, TEXT("0") },
+			{ TEXT("vertices"), TEXT("array"), TEXT("New vertex positions [{x, y, z}, ...]"), true },
+			{ TEXT("triangles"), TEXT("array"), TEXT("New triangle indices [v0, v1, v2, ...]"), true },
+			{ TEXT("uvs"), TEXT("array"), TEXT("New UV coordinates"), false },
+			{ TEXT("normals"), TEXT("array"), TEXT("New normals"), false },
+			{ TEXT("vertex_colors"), TEXT("array"), TEXT("New vertex colors"), false },
+		};
+	}
+	
+	virtual FECACommandResult Execute(const TSharedPtr<FJsonObject>& Params) override;
+};
+
+//==============================================================================
+// MESH COMBINE / MERGE
+//==============================================================================
+
+/**
+ * Combine multiple meshes into one
+ */
+class FECACommand_MeshCombine : public IECACommand
+{
+public:
+	virtual FString GetName() const override { return TEXT("mesh_combine"); }
+	virtual FString GetDescription() const override { return TEXT("Combine multiple meshes into a single mesh asset"); }
+	virtual FString GetCategory() const override { return TEXT("Mesh"); }
+	
+	virtual TArray<FECACommandParam> GetParameters() const override
+	{
+		return {
+			{ TEXT("output_path"), TEXT("string"), TEXT("Path for the combined mesh asset"), true },
+			{ TEXT("mesh_paths"), TEXT("array"), TEXT("Array of mesh asset paths to combine"), true },
+			{ TEXT("transforms"), TEXT("array"), TEXT("Array of transforms for each mesh [{location, rotation, scale}, ...]"), false },
+			{ TEXT("merge_materials"), TEXT("boolean"), TEXT("Merge material slots"), false, TEXT("false") },
+		};
+	}
+	
+	virtual FECACommandResult Execute(const TSharedPtr<FJsonObject>& Params) override;
+};
+
+//==============================================================================
+// IMPORT / EXPORT
+//==============================================================================
+
+/**
+ * Import mesh from file (OBJ already exists in ECAAssetCommands, this adds more formats)
+ */
+class FECACommand_ImportMesh : public IECACommand
+{
+public:
+	virtual FString GetName() const override { return TEXT("import_mesh"); }
+	virtual FString GetDescription() const override { return TEXT("Import mesh from OBJ, FBX, or glTF file"); }
+	virtual FString GetCategory() const override { return TEXT("Mesh"); }
+	
+	virtual TArray<FECACommandParam> GetParameters() const override
+	{
+		return {
+			{ TEXT("input_file"), TEXT("string"), TEXT("Input file path (.obj, .fbx, .gltf, .glb)"), true },
+			{ TEXT("asset_path"), TEXT("string"), TEXT("Destination asset path in UE"), true },
+			{ TEXT("scale"), TEXT("number"), TEXT("Import scale factor"), false, TEXT("1") },
+			{ TEXT("import_materials"), TEXT("boolean"), TEXT("Import materials if available"), false, TEXT("true") },
+			{ TEXT("generate_lightmap_uvs"), TEXT("boolean"), TEXT("Generate lightmap UVs"), false, TEXT("true") },
+		};
+	}
+	
+	virtual FECACommandResult Execute(const TSharedPtr<FJsonObject>& Params) override;
+};
+
+/**
+ * Export mesh to file
+ */
+class FECACommand_ExportMesh : public IECACommand
+{
+public:
+	virtual FString GetName() const override { return TEXT("export_mesh"); }
+	virtual FString GetDescription() const override { return TEXT("Export mesh asset to OBJ or FBX file"); }
+	virtual FString GetCategory() const override { return TEXT("Mesh"); }
+	
+	virtual TArray<FECACommandParam> GetParameters() const override
+	{
+		return {
+			{ TEXT("mesh_path"), TEXT("string"), TEXT("Path to the mesh asset to export"), true },
+			{ TEXT("output_file"), TEXT("string"), TEXT("Output file path (.obj or .fbx)"), true },
+			{ TEXT("export_materials"), TEXT("boolean"), TEXT("Export materials/textures"), false, TEXT("false") },
+		};
+	}
+	
+	virtual FECACommandResult Execute(const TSharedPtr<FJsonObject>& Params) override;
+};
+
+//==============================================================================
+// VORONOI / FRACTURE
+//==============================================================================
+
+/**
+ * Voronoi fracture mesh
+ */
+class FECACommand_MeshVoronoiFracture : public IECACommand
+{
+public:
+	virtual FString GetName() const override { return TEXT("mesh_voronoi_fracture"); }
+	virtual FString GetDescription() const override { return TEXT("Fracture mesh using Voronoi cells"); }
+	virtual FString GetCategory() const override { return TEXT("Mesh"); }
+	
+	virtual TArray<FECACommandParam> GetParameters() const override
+	{
+		return {
+			{ TEXT("mesh_path"), TEXT("string"), TEXT("Path to the mesh asset to fracture"), true },
+			{ TEXT("output_path_prefix"), TEXT("string"), TEXT("Prefix path for output pieces (e.g., /Game/Meshes/Piece_)"), true },
+			{ TEXT("num_cells"), TEXT("number"), TEXT("Number of Voronoi cells"), false, TEXT("10") },
+			{ TEXT("site_points"), TEXT("array"), TEXT("Optional explicit Voronoi site points [{x,y,z}, ...]"), false },
+			{ TEXT("random_seed"), TEXT("number"), TEXT("Random seed for cell generation"), false, TEXT("0") },
+		};
+	}
+	
+	virtual FECACommandResult Execute(const TSharedPtr<FJsonObject>& Params) override;
+};
+
+/**
+ * Displace mesh vertices based on depth image difference
+ * 
+ * WORKFLOW:
+ * 1. Spawn mesh as actor in the world (use ue_spawn_actor with a StaticMeshActor)
+ * 2. Position editor camera to view the mesh from desired angle
+ * 3. Take a depth screenshot (use ue_take_depth_screenshot, save to file)
+ * 4. Edit the depth image externally (paint darker areas to push toward camera, lighter to push away)
+ *    - Use ue_edit_image or external tools like Photoshop
+ *    - Black = closest to camera, White = farthest from camera
+ * 5. Call this tool with the edited depth image as target_depth_image
+ * 6. View the result in editor, iterate as needed
+ * 
+ * DEPTH IMAGE FORMAT:
+ * - Grayscale PNG where pixel brightness = distance from camera
+ * - Black (0) = surface at camera position
+ * - White (255) = surface at max_depth distance
+ * - Paint DARKER to pull surface TOWARD camera
+ * - Paint LIGHTER to push surface AWAY from camera
+ */
+class FECACommand_MeshDisplaceFromDepth : public IECACommand
+{
+public:
+	virtual FString GetName() const override { return TEXT("mesh_displace_from_depth"); }
+	virtual FString GetDescription() const override { return TEXT("Displace mesh vertices based on depth image. WORKFLOW: 1) Spawn mesh in world, 2) Position camera, 3) Take depth screenshot, 4) Edit depth image (darker=closer), 5) Apply displacement, 6) Iterate."); }
+	virtual FString GetCategory() const override { return TEXT("Mesh"); }
+	
+	virtual TArray<FECACommandParam> GetParameters() const override
+	{
+		return {
+			{ TEXT("mesh_path"), TEXT("string"), TEXT("Path to the mesh asset to displace"), true },
+			{ TEXT("actor_path"), TEXT("string"), TEXT("Path to the actor instance in the world (e.g. '/Game/Maps/MyLevel.MyLevel:PersistentLevel.StaticMeshActor_0'). Required to get world transform for correct projection."), true },
+			{ TEXT("output_path"), TEXT("string"), TEXT("Output path (defaults to modifying in place)"), false },
+			{ TEXT("target_depth_image"), TEXT("string"), TEXT("Path to target depth PNG (grayscale: black=close, white=far). If not provided, displaces toward camera."), false },
+			{ TEXT("strength"), TEXT("number"), TEXT("Displacement strength - 255 in the image = this many world units"), false, TEXT("100") },
+			{ TEXT("use_absolute"), TEXT("boolean"), TEXT("If true, use target image as absolute heightmap (ignore current depth). Simpler and more predictable."), false, TEXT("true") },
+			{ TEXT("center_value"), TEXT("number"), TEXT("Grayscale value that means zero displacement (0-255). Default 128 = middle gray."), false, TEXT("128") },
+			{ TEXT("max_depth"), TEXT("number"), TEXT("Maximum depth for normalization (in units) - only used if use_absolute=false"), false, TEXT("10000") },
+			{ TEXT("resolution"), TEXT("number"), TEXT("Depth capture resolution"), false, TEXT("512") },
+			{ TEXT("invert"), TEXT("boolean"), TEXT("Invert displacement direction"), false, TEXT("false") },
+			{ TEXT("smooth_iterations"), TEXT("number"), TEXT("Number of smoothing passes on displacement values (0=none, 3-5 recommended)"), false, TEXT("3") },
+		};
+	}
+	
+	virtual FECACommandResult Execute(const TSharedPtr<FJsonObject>& Params) override;
+};
+
+/**
+ * Displace mesh vertices using a heightmap texture sampled via UV coordinates.
+ * Much simpler and more predictable than depth-based displacement.
+ * 
+ * WORKFLOW:
+ * 1. Create/obtain a grayscale heightmap image
+ * 2. Ensure your mesh has proper UV coordinates
+ * 3. Call this tool to displace vertices along their normals
+ * 
+ * HEIGHTMAP FORMAT:
+ * - Black (0) = no displacement (or negative if center_value > 0)
+ * - White (255) = maximum displacement
+ * - Use center_value to set the "zero point" (default 0 = black is zero)
+ */
+class FECACommand_MeshDisplaceFromHeightmap : public IECACommand
+{
+public:
+	virtual FString GetName() const override { return TEXT("mesh_displace_from_heightmap"); }
+	virtual FString GetDescription() const override { return TEXT("Displace mesh vertices along normals using a heightmap texture sampled via UVs. Simpler and more predictable than depth-based displacement."); }
+	virtual FString GetCategory() const override { return TEXT("Mesh"); }
+	
+	virtual TArray<FECACommandParam> GetParameters() const override
+	{
+		return {
+			{ TEXT("mesh_path"), TEXT("string"), TEXT("Path to the mesh asset to displace"), true },
+			{ TEXT("heightmap_path"), TEXT("string"), TEXT("Path to heightmap image (grayscale PNG). Black=min, White=max displacement."), true },
+			{ TEXT("output_path"), TEXT("string"), TEXT("Output path (defaults to modifying in place)"), false },
+			{ TEXT("strength"), TEXT("number"), TEXT("Maximum displacement in world units (when heightmap = 255)"), false, TEXT("100") },
+			{ TEXT("center_value"), TEXT("number"), TEXT("Grayscale value that means zero displacement (0-255). Default 0 = black is zero, 128 = mid-gray is zero."), false, TEXT("0") },
+			{ TEXT("uv_channel"), TEXT("number"), TEXT("UV channel to use for sampling (0-7)"), false, TEXT("0") },
+			{ TEXT("invert"), TEXT("boolean"), TEXT("Invert displacement direction"), false, TEXT("false") },
+			{ TEXT("smooth_iterations"), TEXT("number"), TEXT("Number of smoothing passes (0=none, 3-5 recommended)"), false, TEXT("0") },
+		};
+	}
+	
+	virtual FECACommandResult Execute(const TSharedPtr<FJsonObject>& Params) override;
+};
diff --git a/Engine/Plugins/Experimental/ECABridge/Source/ECABridge/Public/Commands/ECAMetasoundCommands.h b/Engine/Plugins/Experimental/ECABridge/Source/ECABridge/Public/Commands/ECAMetasoundCommands.h
new file mode 100644
--- /dev/null
+++ b/Engine/Plugins/Experimental/ECABridge/Source/ECABridge/Public/Commands/ECAMetasoundCommands.h
@@ -0,0 +1,350 @@
+// Copyright Epic Games, Inc. All Rights Reserved.
+
+#pragma once
+
+#include "Commands/ECACommand.h"
+
+/**
+ * Get available MetaSound sources in the project
+ */
+class FECACommand_GetMetasoundSources : public IECACommand
+{
+public:
+	virtual FString GetName() const override { return TEXT("get_metasound_sources"); }
+	virtual FString GetDescription() const override { return TEXT("Get all available MetaSound source assets in the project"); }
+	virtual FString GetCategory() const override { return TEXT("Metasound"); }
+	
+	virtual TArray<FECACommandParam> GetParameters() const override
+	{
+		return {
+			{ TEXT("path_filter"), TEXT("string"), TEXT("Filter sources by path (partial match)"), false },
+			{ TEXT("name_filter"), TEXT("string"), TEXT("Filter sources by name (partial match)"), false }
+		};
+	}
+	
+	virtual FECACommandResult Execute(const TSharedPtr<FJsonObject>& Params) override;
+};
+
+/**
+ * Create a new MetaSound source asset
+ */
+class FECACommand_CreateMetasoundSource : public IECACommand
+{
+public:
+	virtual FString GetName() const override { return TEXT("create_metasound_source"); }
+	virtual FString GetDescription() const override { return TEXT("Create a new MetaSound source asset"); }
+	virtual FString GetCategory() const override { return TEXT("Metasound"); }
+	
+	virtual TArray<FECACommandParam> GetParameters() const override
+	{
+		return {
+			{ TEXT("asset_path"), TEXT("string"), TEXT("Path for the new asset (e.g., /Game/Audio/MS_NewSound)"), true },
+			{ TEXT("output_format"), TEXT("string"), TEXT("Output format: Mono or Stereo"), false, TEXT("Stereo") }
+		};
+	}
+	
+	virtual FECACommandResult Execute(const TSharedPtr<FJsonObject>& Params) override;
+};
+
+/**
+ * Get all nodes in a MetaSound source
+ */
+class FECACommand_GetMetasoundNodes : public IECACommand
+{
+public:
+	virtual FString GetName() const override { return TEXT("get_metasound_nodes"); }
+	virtual FString GetDescription() const override { return TEXT("Get all nodes in a MetaSound source graph"); }
+	virtual FString GetCategory() const override { return TEXT("Metasound"); }
+	
+	virtual TArray<FECACommandParam> GetParameters() const override
+	{
+		return {
+			{ TEXT("asset_path"), TEXT("string"), TEXT("Path to the MetaSound source asset"), true },
+			{ TEXT("include_details"), TEXT("boolean"), TEXT("Include detailed node information (pins, values)"), false, TEXT("false") }
+		};
+	}
+	
+	virtual FECACommandResult Execute(const TSharedPtr<FJsonObject>& Params) override;
+};
+
+/**
+ * Add a node to a MetaSound source
+ */
+class FECACommand_AddMetasoundNode : public IECACommand
+{
+public:
+	virtual FString GetName() const override { return TEXT("add_metasound_node"); }
+	virtual FString GetDescription() const override { return TEXT("Add a node to a MetaSound source graph"); }
+	virtual FString GetCategory() const override { return TEXT("Metasound"); }
+	
+	virtual TArray<FECACommandParam> GetParameters() const override
+	{
+		return {
+			{ TEXT("asset_path"), TEXT("string"), TEXT("Path to the MetaSound source asset"), true },
+			{ TEXT("node_class"), TEXT("string"), TEXT("Node class name (e.g., 'Oscillator', 'AD Envelope', 'Trigger Repeat')"), true },
+			{ TEXT("node_name"), TEXT("string"), TEXT("Optional display name for the node"), false },
+			{ TEXT("position_x"), TEXT("number"), TEXT("X position in the graph"), false, TEXT("0") },
+			{ TEXT("position_y"), TEXT("number"), TEXT("Y position in the graph"), false, TEXT("0") }
+		};
+	}
+	
+	virtual FECACommandResult Execute(const TSharedPtr<FJsonObject>& Params) override;
+};
+
+/**
+ * Remove a node from a MetaSound source
+ */
+class FECACommand_RemoveMetasoundNode : public IECACommand
+{
+public:
+	virtual FString GetName() const override { return TEXT("remove_metasound_node"); }
+	virtual FString GetDescription() const override { return TEXT("Remove a node from a MetaSound source graph"); }
+	virtual FString GetCategory() const override { return TEXT("Metasound"); }
+	
+	virtual TArray<FECACommandParam> GetParameters() const override
+	{
+		return {
+			{ TEXT("asset_path"), TEXT("string"), TEXT("Path to the MetaSound source asset"), true },
+			{ TEXT("node_id"), TEXT("string"), TEXT("ID of the node to remove"), true }
+		};
+	}
+	
+	virtual FECACommandResult Execute(const TSharedPtr<FJsonObject>& Params) override;
+};
+
+/**
+ * Connect two nodes in a MetaSound source
+ */
+class FECACommand_ConnectMetasoundNodes : public IECACommand
+{
+public:
+	virtual FString GetName() const override { return TEXT("connect_metasound_nodes"); }
+	virtual FString GetDescription() const override { return TEXT("Connect two nodes in a MetaSound source graph"); }
+	virtual FString GetCategory() const override { return TEXT("Metasound"); }
+	
+	virtual TArray<FECACommandParam> GetParameters() const override
+	{
+		return {
+			{ TEXT("asset_path"), TEXT("string"), TEXT("Path to the MetaSound source asset"), true },
+			{ TEXT("source_node_id"), TEXT("string"), TEXT("ID of the source node"), true },
+			{ TEXT("source_pin"), TEXT("string"), TEXT("Name of the output pin on the source node"), true },
+			{ TEXT("target_node_id"), TEXT("string"), TEXT("ID of the target node"), true },
+			{ TEXT("target_pin"), TEXT("string"), TEXT("Name of the input pin on the target node"), true }
+		};
+	}
+	
+	virtual FECACommandResult Execute(const TSharedPtr<FJsonObject>& Params) override;
+};
+
+/**
+ * Disconnect nodes in a MetaSound source
+ */
+class FECACommand_DisconnectMetasoundNodes : public IECACommand
+{
+public:
+	virtual FString GetName() const override { return TEXT("disconnect_metasound_nodes"); }
+	virtual FString GetDescription() const override { return TEXT("Disconnect a connection between nodes in a MetaSound source graph"); }
+	virtual FString GetCategory() const override { return TEXT("Metasound"); }
+	
+	virtual TArray<FECACommandParam> GetParameters() const override
+	{
+		return {
+			{ TEXT("asset_path"), TEXT("string"), TEXT("Path to the MetaSound source asset"), true },
+			{ TEXT("source_node_id"), TEXT("string"), TEXT("ID of the source node"), true },
+			{ TEXT("source_pin"), TEXT("string"), TEXT("Name of the output pin on the source node"), true },
+			{ TEXT("target_node_id"), TEXT("string"), TEXT("ID of the target node"), true },
+			{ TEXT("target_pin"), TEXT("string"), TEXT("Name of the input pin on the target node"), true }
+		};
+	}
+	
+	virtual FECACommandResult Execute(const TSharedPtr<FJsonObject>& Params) override;
+};
+
+/**
+ * Set an input value on a MetaSound node
+ */
+class FECACommand_SetMetasoundNodeInput : public IECACommand
+{
+public:
+	virtual FString GetName() const override { return TEXT("set_metasound_node_input"); }
+	virtual FString GetDescription() const override { return TEXT("Set an input value on a MetaSound node"); }
+	virtual FString GetCategory() const override { return TEXT("Metasound"); }
+	
+	virtual TArray<FECACommandParam> GetParameters() const override
+	{
+		return {
+			{ TEXT("asset_path"), TEXT("string"), TEXT("Path to the MetaSound source asset"), true },
+			{ TEXT("node_id"), TEXT("string"), TEXT("ID of the node"), true },
+			{ TEXT("input_name"), TEXT("string"), TEXT("Name of the input pin"), true },
+			{ TEXT("value"), TEXT("any"), TEXT("Value to set (type depends on input)"), true }
+		};
+	}
+	
+	virtual FECACommandResult Execute(const TSharedPtr<FJsonObject>& Params) override;
+};
+
+/**
+ * Add an input to a MetaSound source
+ */
+class FECACommand_AddMetasoundInput : public IECACommand
+{
+public:
+	virtual FString GetName() const override { return TEXT("add_metasound_input"); }
+	virtual FString GetDescription() const override { return TEXT("Add an input to a MetaSound source"); }
+	virtual FString GetCategory() const override { return TEXT("Metasound"); }
+	
+	virtual TArray<FECACommandParam> GetParameters() const override
+	{
+		return {
+			{ TEXT("asset_path"), TEXT("string"), TEXT("Path to the MetaSound source asset"), true },
+			{ TEXT("input_name"), TEXT("string"), TEXT("Name for the input"), true },
+			{ TEXT("data_type"), TEXT("string"), TEXT("Data type (e.g., Float, Int32, Bool, Trigger, Audio)"), true },
+			{ TEXT("default_value"), TEXT("any"), TEXT("Default value for the input"), false }
+		};
+	}
+	
+	virtual FECACommandResult Execute(const TSharedPtr<FJsonObject>& Params) override;
+};
+
+/**
+ * Remove an input from a MetaSound source
+ */
+class FECACommand_RemoveMetasoundInput : public IECACommand
+{
+public:
+	virtual FString GetName() const override { return TEXT("remove_metasound_input"); }
+	virtual FString GetDescription() const override { return TEXT("Remove an input from a MetaSound source"); }
+	virtual FString GetCategory() const override { return TEXT("Metasound"); }
+	
+	virtual TArray<FECACommandParam> GetParameters() const override
+	{
+		return {
+			{ TEXT("asset_path"), TEXT("string"), TEXT("Path to the MetaSound source asset"), true },
+			{ TEXT("input_name"), TEXT("string"), TEXT("Name of the input to remove"), true }
+		};
+	}
+	
+	virtual FECACommandResult Execute(const TSharedPtr<FJsonObject>& Params) override;
+};
+
+/**
+ * Get available MetaSound node types
+ */
+class FECACommand_GetMetasoundNodeTypes : public IECACommand
+{
+public:
+	virtual FString GetName() const override { return TEXT("get_metasound_node_types"); }
+	virtual FString GetDescription() const override { return TEXT("Get all available MetaSound node types that can be added to a graph"); }
+	virtual FString GetCategory() const override { return TEXT("Metasound"); }
+	
+	virtual TArray<FECACommandParam> GetParameters() const override
+	{
+		return {
+			{ TEXT("category_filter"), TEXT("string"), TEXT("Filter by category (e.g., 'Generators', 'Envelopes', 'Filters')"), false },
+			{ TEXT("name_filter"), TEXT("string"), TEXT("Filter by name (partial match)"), false }
+		};
+	}
+	
+	virtual FECACommandResult Execute(const TSharedPtr<FJsonObject>& Params) override;
+};
+
+/**
+ * Play/preview a MetaSound in the editor
+ */
+class FECACommand_PreviewMetasound : public IECACommand
+{
+public:
+	virtual FString GetName() const override { return TEXT("preview_metasound"); }
+	virtual FString GetDescription() const override { return TEXT("Play/preview a MetaSound asset in the editor"); }
+	virtual FString GetCategory() const override { return TEXT("Metasound"); }
+	
+	virtual TArray<FECACommandParam> GetParameters() const override
+	{
+		return {
+			{ TEXT("asset_path"), TEXT("string"), TEXT("Path to the MetaSound source asset"), true },
+			{ TEXT("action"), TEXT("string"), TEXT("Action: play, stop"), false, TEXT("play") }
+		};
+	}
+	
+	virtual FECACommandResult Execute(const TSharedPtr<FJsonObject>& Params) override;
+};
+
+/**
+ * Get MetaSound source inputs and outputs
+ */
+class FECACommand_GetMetasoundInterface : public IECACommand
+{
+public:
+	virtual FString GetName() const override { return TEXT("get_metasound_interface"); }
+	virtual FString GetDescription() const override { return TEXT("Get the inputs and outputs of a MetaSound source"); }
+	virtual FString GetCategory() const override { return TEXT("Metasound"); }
+	
+	virtual TArray<FECACommandParam> GetParameters() const override
+	{
+		return {
+			{ TEXT("asset_path"), TEXT("string"), TEXT("Path to the MetaSound source asset"), true }
+		};
+	}
+	
+	virtual FECACommandResult Execute(const TSharedPtr<FJsonObject>& Params) override;
+};
+
+/**
+ * Spawn a MetaSound player in the level
+ */
+class FECACommand_SpawnMetasoundPlayer : public IECACommand
+{
+public:
+	virtual FString GetName() const override { return TEXT("spawn_metasound_player"); }
+	virtual FString GetDescription() const override { return TEXT("Spawn an audio component playing a MetaSound in the level"); }
+	virtual FString GetCategory() const override { return TEXT("Metasound"); }
+	
+	virtual TArray<FECACommandParam> GetParameters() const override
+	{
+		return {
+			{ TEXT("asset_path"), TEXT("string"), TEXT("Path to the MetaSound source asset"), true },
+			{ TEXT("location"), TEXT("object"), TEXT("Spawn location {x, y, z}"), false },
+			{ TEXT("name"), TEXT("string"), TEXT("Display name for the actor"), false },
+			{ TEXT("auto_play"), TEXT("boolean"), TEXT("Whether to auto-play the sound"), false, TEXT("true") },
+			{ TEXT("folder"), TEXT("string"), TEXT("Outliner folder path"), false }
+		};
+	}
+	
+	virtual FECACommandResult Execute(const TSharedPtr<FJsonObject>& Params) override;
+};
+
+
+/**
+ * Auto-layout MetaSound nodes for visual clarity
+ * 
+ * Layout strategies:
+ * - horizontal: Arrange nodes left-to-right from inputs to outputs (default)
+ * - vertical: Arrange nodes top-to-bottom from inputs to outputs
+ * - tree: Arrange as a tree structure
+ * - compact: Minimize total graph area while maintaining readability
+ * 
+ * The algorithm:
+ * 1. Identifies input nodes (audio inputs, parameters)
+ * 2. Traverses connections forward to determine node order/depth
+ * 3. Positions nodes with consistent spacing to avoid overlaps
+ */
+class FECACommand_AutoLayoutMetasoundGraph : public IECACommand
+{
+public:
+	virtual FString GetName() const override { return TEXT("auto_layout_metasound_graph"); }
+	virtual FString GetDescription() const override { return TEXT("Automatically arrange MetaSound nodes for visual clarity. Organizes by connection flow with consistent spacing."); }
+	virtual FString GetCategory() const override { return TEXT("Metasound"); }
+	
+	virtual TArray<FECACommandParam> GetParameters() const override
+	{
+		return {
+			{ TEXT("asset_path"), TEXT("string"), TEXT("Path to the MetaSound source asset"), true },
+			{ TEXT("strategy"), TEXT("string"), TEXT("Layout strategy: horizontal (default), vertical, tree, compact"), false, TEXT("horizontal") },
+			{ TEXT("spacing_x"), TEXT("number"), TEXT("Horizontal spacing between nodes (default: 300)"), false, TEXT("300") },
+			{ TEXT("spacing_y"), TEXT("number"), TEXT("Vertical spacing between nodes (default: 100)"), false, TEXT("100") },
+			{ TEXT("node_ids"), TEXT("array"), TEXT("Specific node IDs to layout (optional, layouts all if not specified)"), false }
+		};
+	}
+	
+	virtual FECACommandResult Execute(const TSharedPtr<FJsonObject>& Params) override;
+};
diff --git a/Engine/Plugins/Experimental/ECABridge/Source/ECABridge/Public/Commands/ECANiagaraCommands.h b/Engine/Plugins/Experimental/ECABridge/Source/ECABridge/Public/Commands/ECANiagaraCommands.h
new file mode 100644
--- /dev/null
+++ b/Engine/Plugins/Experimental/ECABridge/Source/ECABridge/Public/Commands/ECANiagaraCommands.h
@@ -0,0 +1,532 @@
+// Copyright Epic Games, Inc. All Rights Reserved.
+
+#pragma once
+
+#include "Commands/ECACommand.h"
+
+/**
+ * Get available Niagara systems in the project
+ */
+class FECACommand_GetNiagaraSystems : public IECACommand
+{
+public:
+	virtual FString GetName() const override { return TEXT("get_niagara_systems"); }
+	virtual FString GetDescription() const override { return TEXT("Get all available Niagara systems in the project"); }
+	virtual FString GetCategory() const override { return TEXT("Niagara"); }
+	
+	virtual TArray<FECACommandParam> GetParameters() const override
+	{
+		return {
+			{ TEXT("path_filter"), TEXT("string"), TEXT("Filter systems by path (partial match)"), false },
+			{ TEXT("name_filter"), TEXT("string"), TEXT("Filter systems by name (partial match)"), false }
+		};
+	}
+	
+	virtual FECACommandResult Execute(const TSharedPtr<FJsonObject>& Params) override;
+};
+
+/**
+ * Spawn a Niagara effect in the level
+ */
+class FECACommand_SpawnNiagaraEffect : public IECACommand
+{
+public:
+	virtual FString GetName() const override { return TEXT("spawn_niagara_effect"); }
+	virtual FString GetDescription() const override { return TEXT("Spawn a Niagara particle effect in the level"); }
+	virtual FString GetCategory() const override { return TEXT("Niagara"); }
+	
+	virtual TArray<FECACommandParam> GetParameters() const override
+	{
+		return {
+			{ TEXT("system_path"), TEXT("string"), TEXT("Path to the Niagara system asset (e.g., /Game/Effects/NS_Fire)"), true },
+			{ TEXT("location"), TEXT("object"), TEXT("Spawn location {x, y, z}"), false },
+			{ TEXT("rotation"), TEXT("object"), TEXT("Spawn rotation {pitch, yaw, roll}"), false },
+			{ TEXT("scale"), TEXT("object"), TEXT("Scale {x, y, z}"), false },
+			{ TEXT("name"), TEXT("string"), TEXT("Display name for the effect actor"), false },
+			{ TEXT("auto_activate"), TEXT("boolean"), TEXT("Whether to auto-activate the effect"), false, TEXT("true") },
+			{ TEXT("folder"), TEXT("string"), TEXT("Outliner folder path"), false }
+		};
+	}
+	
+	virtual FECACommandResult Execute(const TSharedPtr<FJsonObject>& Params) override;
+};
+
+/**
+ * Create a new Niagara system asset
+ */
+class FECACommand_CreateNiagaraSystem : public IECACommand
+{
+public:
+	virtual FString GetName() const override { return TEXT("create_niagara_system"); }
+	virtual FString GetDescription() const override { return TEXT("Create a new Niagara system asset"); }
+	virtual FString GetCategory() const override { return TEXT("Niagara"); }
+	
+	virtual TArray<FECACommandParam> GetParameters() const override
+	{
+		return {
+			{ TEXT("asset_path"), TEXT("string"), TEXT("Path for the new asset (e.g., /Game/Effects/NS_NewEffect)"), true },
+			{ TEXT("template"), TEXT("string"), TEXT("Template to use: empty, sprite, mesh, ribbon, or path to existing system"), false, TEXT("empty") }
+		};
+	}
+	
+	virtual FECACommandResult Execute(const TSharedPtr<FJsonObject>& Params) override;
+};
+
+/**
+ * Get Niagara effect actors in the level
+ */
+class FECACommand_GetNiagaraActors : public IECACommand
+{
+public:
+	virtual FString GetName() const override { return TEXT("get_niagara_actors"); }
+	virtual FString GetDescription() const override { return TEXT("Get all Niagara effect actors in the current level"); }
+	virtual FString GetCategory() const override { return TEXT("Niagara"); }
+	
+	virtual TArray<FECACommandParam> GetParameters() const override
+	{
+		return {
+			{ TEXT("system_filter"), TEXT("string"), TEXT("Filter by system name (partial match)"), false }
+		};
+	}
+	
+	virtual FECACommandResult Execute(const TSharedPtr<FJsonObject>& Params) override;
+};
+
+/**
+ * Set a parameter on a Niagara component
+ */
+class FECACommand_SetNiagaraParameter : public IECACommand
+{
+public:
+	virtual FString GetName() const override { return TEXT("set_niagara_parameter"); }
+	virtual FString GetDescription() const override { return TEXT("Set a user parameter on a Niagara effect"); }
+	virtual FString GetCategory() const override { return TEXT("Niagara"); }
+	
+	virtual TArray<FECACommandParam> GetParameters() const override
+	{
+		return {
+			{ TEXT("actor_name"), TEXT("string"), TEXT("Name of the Niagara actor"), true },
+			{ TEXT("parameter_name"), TEXT("string"), TEXT("Name of the parameter to set"), true },
+			{ TEXT("parameter_type"), TEXT("string"), TEXT("Type: float, int, bool, vector, color, linear_color"), true },
+			{ TEXT("value"), TEXT("any"), TEXT("Value to set (number, boolean, or {x,y,z}/{r,g,b,a} object)"), true }
+		};
+	}
+	
+	virtual FECACommandResult Execute(const TSharedPtr<FJsonObject>& Params) override;
+};
+
+/**
+ * Control Niagara effect playback
+ */
+class FECACommand_ControlNiagaraEffect : public IECACommand
+{
+public:
+	virtual FString GetName() const override { return TEXT("control_niagara_effect"); }
+	virtual FString GetDescription() const override { return TEXT("Control playback of a Niagara effect (activate, deactivate, reset)"); }
+	virtual FString GetCategory() const override { return TEXT("Niagara"); }
+	
+	virtual TArray<FECACommandParam> GetParameters() const override
+	{
+		return {
+			{ TEXT("actor_name"), TEXT("string"), TEXT("Name of the Niagara actor"), true },
+			{ TEXT("action"), TEXT("string"), TEXT("Action to perform: activate, deactivate, reset, reset_system"), true }
+		};
+	}
+	
+	virtual FECACommandResult Execute(const TSharedPtr<FJsonObject>& Params) override;
+};
+
+/**
+ * Get Niagara system user parameters
+ */
+class FECACommand_GetNiagaraParameters : public IECACommand
+{
+public:
+	virtual FString GetName() const override { return TEXT("get_niagara_parameters"); }
+	virtual FString GetDescription() const override { return TEXT("Get user-exposed parameters of a Niagara system"); }
+	virtual FString GetCategory() const override { return TEXT("Niagara"); }
+	
+	virtual TArray<FECACommandParam> GetParameters() const override
+	{
+		return {
+			{ TEXT("actor_name"), TEXT("string"), TEXT("Name of the Niagara actor in the level"), false },
+			{ TEXT("system_path"), TEXT("string"), TEXT("Path to the Niagara system asset"), false }
+		};
+	}
+	
+	virtual FECACommandResult Execute(const TSharedPtr<FJsonObject>& Params) override;
+};
+
+/**
+ * Add a Niagara component to an existing actor
+ */
+class FECACommand_AddNiagaraComponent : public IECACommand
+{
+public:
+	virtual FString GetName() const override { return TEXT("add_niagara_component"); }
+	virtual FString GetDescription() const override { return TEXT("Add a Niagara component to an existing actor"); }
+	virtual FString GetCategory() const override { return TEXT("Niagara"); }
+	
+	virtual TArray<FECACommandParam> GetParameters() const override
+	{
+		return {
+			{ TEXT("actor_name"), TEXT("string"), TEXT("Name of the actor to add the component to"), true },
+			{ TEXT("system_path"), TEXT("string"), TEXT("Path to the Niagara system asset"), true },
+			{ TEXT("component_name"), TEXT("string"), TEXT("Name for the new component"), false },
+			{ TEXT("relative_location"), TEXT("object"), TEXT("Relative location {x, y, z}"), false },
+			{ TEXT("relative_rotation"), TEXT("object"), TEXT("Relative rotation {pitch, yaw, roll}"), false },
+			{ TEXT("auto_activate"), TEXT("boolean"), TEXT("Whether to auto-activate"), false, TEXT("true") }
+		};
+	}
+	
+	virtual FECACommandResult Execute(const TSharedPtr<FJsonObject>& Params) override;
+};
+
+/**
+ * Add an emitter to a Niagara system
+ */
+class FECACommand_AddNiagaraEmitter : public IECACommand
+{
+public:
+	virtual FString GetName() const override { return TEXT("add_niagara_emitter"); }
+	virtual FString GetDescription() const override { return TEXT("Add an emitter to a Niagara system asset"); }
+	virtual FString GetCategory() const override { return TEXT("Niagara"); }
+	
+	virtual TArray<FECACommandParam> GetParameters() const override
+	{
+		return {
+			{ TEXT("system_path"), TEXT("string"), TEXT("Path to the Niagara system asset"), true },
+			{ TEXT("emitter_path"), TEXT("string"), TEXT("Path to emitter asset to add, or 'empty' for blank emitter"), false, TEXT("empty") },
+			{ TEXT("emitter_name"), TEXT("string"), TEXT("Name for the new emitter"), false }
+		};
+	}
+	
+	virtual FECACommandResult Execute(const TSharedPtr<FJsonObject>& Params) override;
+};
+
+/**
+ * Add a module to a Niagara emitter
+ */
+class FECACommand_AddNiagaraModule : public IECACommand
+{
+public:
+	virtual FString GetName() const override { return TEXT("add_niagara_module"); }
+	virtual FString GetDescription() const override { return TEXT("Add a module to a Niagara emitter's script stack"); }
+	virtual FString GetCategory() const override { return TEXT("Niagara"); }
+	
+	virtual TArray<FECACommandParam> GetParameters() const override
+	{
+		return {
+			{ TEXT("system_path"), TEXT("string"), TEXT("Path to the Niagara system asset"), true },
+			{ TEXT("emitter_name"), TEXT("string"), TEXT("Name of the emitter (or index like 'Emitter[0]')"), true },
+			{ TEXT("module_path"), TEXT("string"), TEXT("Path to module script (e.g., '/Niagara/Modules/Spawn/Velocity/AddVelocity.AddVelocity')"), true },
+			{ TEXT("script_usage"), TEXT("string"), TEXT("Target script: emitter_spawn, emitter_update, particle_spawn, particle_update"), true },
+			{ TEXT("index"), TEXT("number"), TEXT("Position to insert module (-1 for end)"), false, TEXT("-1") }
+		};
+	}
+	
+	virtual FECACommandResult Execute(const TSharedPtr<FJsonObject>& Params) override;
+};
+
+/**
+ * Set a module input value in a Niagara emitter
+ */
+class FECACommand_SetNiagaraModuleInput : public IECACommand
+{
+public:
+	virtual FString GetName() const override { return TEXT("set_niagara_module_input"); }
+	virtual FString GetDescription() const override { return TEXT("Set an input value on a Niagara module"); }
+	virtual FString GetCategory() const override { return TEXT("Niagara"); }
+	
+	virtual TArray<FECACommandParam> GetParameters() const override
+	{
+		return {
+			{ TEXT("system_path"), TEXT("string"), TEXT("Path to the Niagara system asset"), true },
+			{ TEXT("emitter_name"), TEXT("string"), TEXT("Name of the emitter"), true },
+			{ TEXT("module_name"), TEXT("string"), TEXT("Name of the module"), true },
+			{ TEXT("input_name"), TEXT("string"), TEXT("Name of the input to set"), true },
+			{ TEXT("value"), TEXT("any"), TEXT("Value to set (number, bool, or {x,y,z}/{r,g,b,a} object)"), true },
+			{ TEXT("script_usage"), TEXT("string"), TEXT("Script containing the module: emitter_spawn, emitter_update, particle_spawn, particle_update"), true }
+		};
+	}
+	
+	virtual FECACommandResult Execute(const TSharedPtr<FJsonObject>& Params) override;
+};
+
+/**
+ * Get emitters in a Niagara system
+ */
+class FECACommand_GetNiagaraEmitters : public IECACommand
+{
+public:
+	virtual FString GetName() const override { return TEXT("get_niagara_emitters"); }
+	virtual FString GetDescription() const override { return TEXT("Get all emitters in a Niagara system"); }
+	virtual FString GetCategory() const override { return TEXT("Niagara"); }
+	
+	virtual TArray<FECACommandParam> GetParameters() const override
+	{
+		return {
+			{ TEXT("system_path"), TEXT("string"), TEXT("Path to the Niagara system asset"), true }
+		};
+	}
+	
+	virtual FECACommandResult Execute(const TSharedPtr<FJsonObject>& Params) override;
+};
+
+/**
+ * Get modules in a Niagara emitter
+ */
+class FECACommand_GetNiagaraModules : public IECACommand
+{
+public:
+	virtual FString GetName() const override { return TEXT("get_niagara_modules"); }
+	virtual FString GetDescription() const override { return TEXT("Get all modules in a Niagara emitter's script stack"); }
+	virtual FString GetCategory() const override { return TEXT("Niagara"); }
+	
+	virtual TArray<FECACommandParam> GetParameters() const override
+	{
+		return {
+			{ TEXT("system_path"), TEXT("string"), TEXT("Path to the Niagara system asset"), true },
+			{ TEXT("emitter_name"), TEXT("string"), TEXT("Name of the emitter"), true },
+			{ TEXT("script_usage"), TEXT("string"), TEXT("Script to query: emitter_spawn, emitter_update, particle_spawn, particle_update, all"), false, TEXT("all") }
+		};
+	}
+	
+	virtual FECACommandResult Execute(const TSharedPtr<FJsonObject>& Params) override;
+};
+
+/**
+ * List available Niagara modules
+ */
+class FECACommand_ListNiagaraModules : public IECACommand
+{
+public:
+	virtual FString GetName() const override { return TEXT("list_niagara_modules"); }
+	virtual FString GetDescription() const override { return TEXT("List available built-in Niagara modules by category"); }
+	virtual FString GetCategory() const override { return TEXT("Niagara"); }
+	
+	virtual TArray<FECACommandParam> GetParameters() const override
+	{
+		return {
+			{ TEXT("category"), TEXT("string"), TEXT("Filter by category: spawn, update, emitter, forces, or 'all'"), false, TEXT("all") },
+			{ TEXT("search"), TEXT("string"), TEXT("Search filter for module names"), false }
+		};
+	}
+	
+	virtual FECACommandResult Execute(const TSharedPtr<FJsonObject>& Params) override;
+};
+
+/**
+ * Remove a module from a Niagara emitter
+ */
+class FECACommand_RemoveNiagaraModule : public IECACommand
+{
+public:
+	virtual FString GetName() const override { return TEXT("remove_niagara_module"); }
+	virtual FString GetDescription() const override { return TEXT("Remove a module from a Niagara emitter"); }
+	virtual FString GetCategory() const override { return TEXT("Niagara"); }
+	
+	virtual TArray<FECACommandParam> GetParameters() const override
+	{
+		return {
+			{ TEXT("system_path"), TEXT("string"), TEXT("Path to the Niagara system asset"), true },
+			{ TEXT("emitter_name"), TEXT("string"), TEXT("Name of the emitter"), true },
+			{ TEXT("module_name"), TEXT("string"), TEXT("Name of the module to remove"), true },
+			{ TEXT("script_usage"), TEXT("string"), TEXT("Script containing the module: emitter_spawn, emitter_update, particle_spawn, particle_update"), true }
+		};
+	}
+	
+	virtual FECACommandResult Execute(const TSharedPtr<FJsonObject>& Params) override;
+};
+
+/**
+ * Delete an emitter from a Niagara system
+ */
+class FECACommand_DeleteNiagaraEmitter : public IECACommand
+{
+public:
+	virtual FString GetName() const override { return TEXT("delete_niagara_emitter"); }
+	virtual FString GetDescription() const override { return TEXT("Delete an emitter from a Niagara system"); }
+	virtual FString GetCategory() const override { return TEXT("Niagara"); }
+	
+	virtual TArray<FECACommandParam> GetParameters() const override
+	{
+		return {
+			{ TEXT("system_path"), TEXT("string"), TEXT("Path to the Niagara system asset"), true },
+			{ TEXT("emitter_name"), TEXT("string"), TEXT("Name of the emitter to delete"), true }
+		};
+	}
+	
+	virtual FECACommandResult Execute(const TSharedPtr<FJsonObject>& Params) override;
+};
+
+/**
+ * Set emitter properties (spawn rate, lifetime, etc.)
+ */
+class FECACommand_SetNiagaraEmitterProperty : public IECACommand
+{
+public:
+	virtual FString GetName() const override { return TEXT("set_niagara_emitter_property"); }
+	virtual FString GetDescription() const override { return TEXT("Set a property on a Niagara emitter"); }
+	virtual FString GetCategory() const override { return TEXT("Niagara"); }
+	
+	virtual TArray<FECACommandParam> GetParameters() const override
+	{
+		return {
+			{ TEXT("system_path"), TEXT("string"), TEXT("Path to the Niagara system asset"), true },
+			{ TEXT("emitter_name"), TEXT("string"), TEXT("Name of the emitter"), true },
+			{ TEXT("property"), TEXT("string"), TEXT("Property to set: sim_target (cpu/gpu), enabled, bounds_mode, local_space, determinism"), true },
+			{ TEXT("value"), TEXT("any"), TEXT("Value to set"), true }
+		};
+	}
+	
+	virtual FECACommandResult Execute(const TSharedPtr<FJsonObject>& Params) override;
+};
+
+/**
+ * List available Niagara emitter templates
+ */
+class FECACommand_ListNiagaraEmitterTemplates : public IECACommand
+{
+public:
+	virtual FString GetName() const override { return TEXT("list_niagara_emitter_templates"); }
+	virtual FString GetDescription() const override { return TEXT("List available Niagara emitter templates that can be used with add_niagara_emitter"); }
+	virtual FString GetCategory() const override { return TEXT("Niagara"); }
+	
+	virtual TArray<FECACommandParam> GetParameters() const override
+	{
+		return {
+			{ TEXT("category"), TEXT("string"), TEXT("Filter by category: all, sprites, meshes, ribbons, beams, behaviors"), false, TEXT("all") }
+		};
+	}
+	
+	virtual FECACommandResult Execute(const TSharedPtr<FJsonObject>& Params) override;
+};
+
+/**
+ * Create a new Niagara emitter asset based on a template
+ */
+class FECACommand_CreateNiagaraEmitter : public IECACommand
+{
+public:
+	virtual FString GetName() const override { return TEXT("create_niagara_emitter"); }
+	virtual FString GetCategory() const override { return TEXT("Niagara"); }
+	virtual FString GetDescription() const override { return TEXT("Create a new Niagara emitter asset based on a template with configurable properties"); }
+	
+	virtual TArray<FECACommandParam> GetParameters() const override
+	{
+		return {
+			{ TEXT("asset_path"), TEXT("string"), TEXT("Path for the new emitter asset (e.g., /Game/Effects/E_MyEmitter)"), true },
+			{ TEXT("template"), TEXT("string"), TEXT("Base template: 'empty', 'sprite', 'mesh', 'ribbon', or full path to existing emitter"), false, TEXT("sprite") },
+			{ TEXT("sim_target"), TEXT("string"), TEXT("Simulation target: 'cpu' or 'gpu'"), false, TEXT("cpu") },
+			{ TEXT("local_space"), TEXT("boolean"), TEXT("Whether to simulate in local space"), false, TEXT("false") },
+			{ TEXT("determinism"), TEXT("boolean"), TEXT("Whether to use deterministic random"), false, TEXT("false") },
+			{ TEXT("bounds_mode"), TEXT("string"), TEXT("Bounds calculation mode: 'dynamic' or 'fixed'"), false, TEXT("dynamic") }
+		};
+	}
+	
+	virtual FECACommandResult Execute(const TSharedPtr<FJsonObject>& Params) override;
+};
+
+/**
+ * Add a renderer to a Niagara emitter
+ */
+class FECACommand_AddNiagaraRenderer : public IECACommand
+{
+public:
+	virtual FString GetName() const override { return TEXT("add_niagara_renderer"); }
+	virtual FString GetCategory() const override { return TEXT("Niagara"); }
+	virtual FString GetDescription() const override { return TEXT("Add a renderer (sprite, mesh, ribbon, light) to a Niagara emitter"); }
+	
+	virtual TArray<FECACommandParam> GetParameters() const override
+	{
+		return {
+			{ TEXT("system_path"), TEXT("string"), TEXT("Path to the Niagara system asset"), true },
+			{ TEXT("emitter_name"), TEXT("string"), TEXT("Name of the emitter"), true },
+			{ TEXT("renderer_type"), TEXT("string"), TEXT("Type: 'sprite', 'mesh', 'ribbon', 'light', 'component'"), true },
+			{ TEXT("renderer_name"), TEXT("string"), TEXT("Name for the new renderer"), false },
+			{ TEXT("material_path"), TEXT("string"), TEXT("Path to material asset to assign"), false },
+			{ TEXT("mesh_path"), TEXT("string"), TEXT("Path to static mesh asset (for mesh renderer)"), false }
+		};
+	}
+	
+	virtual FECACommandResult Execute(const TSharedPtr<FJsonObject>& Params) override;
+};
+
+/**
+ * Set material on a Niagara renderer
+ */
+class FECACommand_SetNiagaraMaterial : public IECACommand
+{
+public:
+	virtual FString GetName() const override { return TEXT("set_niagara_material"); }
+	virtual FString GetCategory() const override { return TEXT("Niagara"); }
+	virtual FString GetDescription() const override { return TEXT("Set a material on a Niagara renderer"); }
+	
+	virtual TArray<FECACommandParam> GetParameters() const override
+	{
+		return {
+			{ TEXT("system_path"), TEXT("string"), TEXT("Path to the Niagara system asset"), true },
+			{ TEXT("emitter_name"), TEXT("string"), TEXT("Name of the emitter"), true },
+			{ TEXT("material_path"), TEXT("string"), TEXT("Path to material asset"), true },
+			{ TEXT("renderer_index"), TEXT("number"), TEXT("Index of the renderer (0-based)"), false, TEXT("0") }
+		};
+	}
+	
+	virtual FECACommandResult Execute(const TSharedPtr<FJsonObject>& Params) override;
+};
+
+/**
+ * Set a curve on a Niagara module input
+ */
+class FECACommand_SetNiagaraCurve : public IECACommand
+{
+public:
+	virtual FString GetName() const override { return TEXT("set_niagara_curve"); }
+	virtual FString GetCategory() const override { return TEXT("Niagara"); }
+	virtual FString GetDescription() const override { return TEXT("Set a float or color curve on a Niagara module input"); }
+	
+	virtual TArray<FECACommandParam> GetParameters() const override
+	{
+		return {
+			{ TEXT("system_path"), TEXT("string"), TEXT("Path to the Niagara system asset"), true },
+			{ TEXT("emitter_name"), TEXT("string"), TEXT("Name of the emitter"), true },
+			{ TEXT("module_name"), TEXT("string"), TEXT("Name of the module"), true },
+			{ TEXT("input_name"), TEXT("string"), TEXT("Name of the input to set"), true },
+			{ TEXT("script_usage"), TEXT("string"), TEXT("Script containing the module: emitter_spawn, emitter_update, particle_spawn, particle_update"), true },
+			{ TEXT("curve_type"), TEXT("string"), TEXT("Type: 'float', 'vector', 'color'"), true },
+			{ TEXT("keys"), TEXT("array"), TEXT("Array of curve keys: [{time: 0, value: 1}, ...] or [{time: 0, r: 1, g: 0, b: 0, a: 1}, ...]"), true }
+		};
+	}
+	
+	virtual FECACommandResult Execute(const TSharedPtr<FJsonObject>& Params) override;
+};
+
+/**
+ * Set a dynamic input (expression or random range) on a module input
+ */
+class FECACommand_SetNiagaraDynamicInput : public IECACommand
+{
+public:
+	virtual FString GetName() const override { return TEXT("set_niagara_dynamic_input"); }
+	virtual FString GetCategory() const override { return TEXT("Niagara"); }
+	virtual FString GetDescription() const override { return TEXT("Set a dynamic input like random range, expression, or link to parameter"); }
+	
+	virtual TArray<FECACommandParam> GetParameters() const override
+	{
+		return {
+			{ TEXT("system_path"), TEXT("string"), TEXT("Path to the Niagara system asset"), true },
+			{ TEXT("emitter_name"), TEXT("string"), TEXT("Name of the emitter"), true },
+			{ TEXT("module_name"), TEXT("string"), TEXT("Name of the module"), true },
+			{ TEXT("input_name"), TEXT("string"), TEXT("Name of the input to set"), true },
+			{ TEXT("script_usage"), TEXT("string"), TEXT("Script containing the module: emitter_spawn, emitter_update, particle_spawn, particle_update"), true },
+			{ TEXT("dynamic_input_type"), TEXT("string"), TEXT("Type: 'random_range', 'uniform_random', 'parameter_link', 'custom_expression'"), true },
+			{ TEXT("min_value"), TEXT("any"), TEXT("Minimum value for random range (number or {x,y,z})"), false },
+			{ TEXT("max_value"), TEXT("any"), TEXT("Maximum value for random range (number or {x,y,z})"), false },
+			{ TEXT("parameter_name"), TEXT("string"), TEXT("Name of parameter to link to"), false },
+			{ TEXT("expression"), TEXT("string"), TEXT("Custom HLSL expression"), false }
+		};
+	}
+	
+	virtual FECACommandResult Execute(const TSharedPtr<FJsonObject>& Params) override;
+};
diff --git a/Engine/Plugins/Experimental/ECABridge/Source/ECABridge/Public/Commands/ECAProjectCommands.h b/Engine/Plugins/Experimental/ECABridge/Source/ECABridge/Public/Commands/ECAProjectCommands.h
new file mode 100644
--- /dev/null
+++ b/Engine/Plugins/Experimental/ECABridge/Source/ECABridge/Public/Commands/ECAProjectCommands.h
@@ -0,0 +1,107 @@
+// Copyright Epic Games, Inc. All Rights Reserved.
+
+#pragma once
+
+#include "Commands/ECACommand.h"
+
+/**
+ * Create an input mapping (action or axis)
+ */
+class FECACommand_CreateInputMapping : public IECACommand
+{
+public:
+	virtual FString GetName() const override { return TEXT("create_input_mapping"); }
+	virtual FString GetDescription() const override { return TEXT("Create an input action or axis mapping in project settings"); }
+	virtual FString GetCategory() const override { return TEXT("Project"); }
+	
+	virtual TArray<FECACommandParam> GetParameters() const override
+	{
+		return {
+			{ TEXT("action_name"), TEXT("string"), TEXT("Name of the input action or axis"), true },
+			{ TEXT("key"), TEXT("string"), TEXT("Key to bind (SpaceBar, LeftMouseButton, W, Gamepad_LeftX, etc.)"), true },
+			{ TEXT("input_type"), TEXT("string"), TEXT("Type of input mapping (Action or Axis)"), false, TEXT("Action") },
+			{ TEXT("scale"), TEXT("number"), TEXT("Scale for axis mappings (-1.0 to 1.0)"), false, TEXT("1.0") },
+			{ TEXT("shift"), TEXT("boolean"), TEXT("Require Shift modifier"), false, TEXT("false") },
+			{ TEXT("ctrl"), TEXT("boolean"), TEXT("Require Ctrl modifier"), false, TEXT("false") },
+			{ TEXT("alt"), TEXT("boolean"), TEXT("Require Alt modifier"), false, TEXT("false") },
+			{ TEXT("cmd"), TEXT("boolean"), TEXT("Require Cmd modifier"), false, TEXT("false") }
+		};
+	}
+	
+	virtual FECACommandResult Execute(const TSharedPtr<FJsonObject>& Params) override;
+};
+
+/**
+ * Get all input mappings
+ */
+class FECACommand_GetInputMappings : public IECACommand
+{
+public:
+	virtual FString GetName() const override { return TEXT("get_input_mappings"); }
+	virtual FString GetDescription() const override { return TEXT("Get all input action and axis mappings"); }
+	virtual FString GetCategory() const override { return TEXT("Project"); }
+	
+	virtual TArray<FECACommandParam> GetParameters() const override
+	{
+		return {
+			{ TEXT("input_type"), TEXT("string"), TEXT("Filter by type (Action, Axis, or All)"), false, TEXT("All") }
+		};
+	}
+	
+	virtual FECACommandResult Execute(const TSharedPtr<FJsonObject>& Params) override;
+};
+
+/**
+ * Remove an input mapping
+ */
+class FECACommand_RemoveInputMapping : public IECACommand
+{
+public:
+	virtual FString GetName() const override { return TEXT("remove_input_mapping"); }
+	virtual FString GetDescription() const override { return TEXT("Remove an input action or axis mapping"); }
+	virtual FString GetCategory() const override { return TEXT("Project"); }
+	
+	virtual TArray<FECACommandParam> GetParameters() const override
+	{
+		return {
+			{ TEXT("action_name"), TEXT("string"), TEXT("Name of the input action or axis to remove"), true },
+			{ TEXT("input_type"), TEXT("string"), TEXT("Type of mapping (Action or Axis)"), false, TEXT("Action") }
+		};
+	}
+	
+	virtual FECACommandResult Execute(const TSharedPtr<FJsonObject>& Params) override;
+};
+
+/**
+ * Get project settings
+ */
+class FECACommand_GetProjectSettings : public IECACommand
+{
+public:
+	virtual FString GetName() const override { return TEXT("get_project_settings"); }
+	virtual FString GetDescription() const override { return TEXT("Get common project settings"); }
+	virtual FString GetCategory() const override { return TEXT("Project"); }
+	
+	virtual FECACommandResult Execute(const TSharedPtr<FJsonObject>& Params) override;
+};
+
+/**
+ * Set a project setting
+ */
+class FECACommand_SetProjectSetting : public IECACommand
+{
+public:
+	virtual FString GetName() const override { return TEXT("set_project_setting"); }
+	virtual FString GetDescription() const override { return TEXT("Set a project setting value"); }
+	virtual FString GetCategory() const override { return TEXT("Project"); }
+	
+	virtual TArray<FECACommandParam> GetParameters() const override
+	{
+		return {
+			{ TEXT("setting_name"), TEXT("string"), TEXT("Name of the setting (e.g., GameDefaultMap, GameInstanceClass)"), true },
+			{ TEXT("setting_value"), TEXT("string"), TEXT("Value to set"), true }
+		};
+	}
+	
+	virtual FECACommandResult Execute(const TSharedPtr<FJsonObject>& Params) override;
+};
