ECABridge experimental MCP plugin

diff --git a/Engine/Plugins/Experimental/ECABridge/Source/ECABridge/Public/Commands/ECAUMGCommands.h b/Engine/Plugins/Experimental/ECABridge/Source/ECABridge/Public/Commands/ECAUMGCommands.h
new file mode 100644
--- /dev/null
+++ b/Engine/Plugins/Experimental/ECABridge/Source/ECABridge/Public/Commands/ECAUMGCommands.h
@@ -0,0 +1,189 @@
+// Copyright Epic Games, Inc. All Rights Reserved.
+
+#pragma once
+
+#include "Commands/ECACommand.h"
+
+/**
+ * Create a new UMG Widget Blueprint
+ */
+class FECACommand_CreateUMGWidgetBlueprint : public IECACommand
+{
+public:
+	virtual FString GetName() const override { return TEXT("create_umg_widget_blueprint"); }
+	virtual FString GetDescription() const override { return TEXT("Create a new UMG Widget Blueprint"); }
+	virtual FString GetCategory() const override { return TEXT("UMG"); }
+	
+	virtual TArray<FECACommandParam> GetParameters() const override
+	{
+		return {
+			{ TEXT("widget_name"), TEXT("string"), TEXT("Name for the new Widget Blueprint"), true },
+			{ TEXT("parent_class"), TEXT("string"), TEXT("Parent class (UserWidget, etc.)"), false, TEXT("UserWidget") },
+			{ TEXT("path"), TEXT("string"), TEXT("Content path for the Widget Blueprint"), false, TEXT("/Game/UI/") }
+		};
+	}
+	
+	virtual FECACommandResult Execute(const TSharedPtr<FJsonObject>& Params) override;
+};
+
+/**
+ * Add a Text Block widget to a Widget Blueprint
+ */
+class FECACommand_AddTextBlockToWidget : public IECACommand
+{
+public:
+	virtual FString GetName() const override { return TEXT("add_text_block_to_widget"); }
+	virtual FString GetDescription() const override { return TEXT("Add a Text Block widget to a Widget Blueprint"); }
+	virtual FString GetCategory() const override { return TEXT("UMG"); }
+	
+	virtual TArray<FECACommandParam> GetParameters() const override
+	{
+		return {
+			{ TEXT("widget_path"), TEXT("string"), TEXT("Path to the Widget Blueprint"), true },
+			{ TEXT("text_block_name"), TEXT("string"), TEXT("Name for the Text Block"), true },
+			{ TEXT("text"), TEXT("string"), TEXT("Initial text content"), false, TEXT("Text") },
+			{ TEXT("position"), TEXT("object"), TEXT("Position in canvas panel {x, y}"), false },
+			{ TEXT("size"), TEXT("object"), TEXT("Size of the text block {width, height}"), false },
+			{ TEXT("font_size"), TEXT("number"), TEXT("Font size in points"), false, TEXT("12") },
+			{ TEXT("color"), TEXT("object"), TEXT("Text color {r, g, b, a} (0.0-1.0)"), false }
+		};
+	}
+	
+	virtual FECACommandResult Execute(const TSharedPtr<FJsonObject>& Params) override;
+};
+
+/**
+ * Add a Button widget to a Widget Blueprint
+ */
+class FECACommand_AddButtonToWidget : public IECACommand
+{
+public:
+	virtual FString GetName() const override { return TEXT("add_button_to_widget"); }
+	virtual FString GetDescription() const override { return TEXT("Add a Button widget to a Widget Blueprint"); }
+	virtual FString GetCategory() const override { return TEXT("UMG"); }
+	
+	virtual TArray<FECACommandParam> GetParameters() const override
+	{
+		return {
+			{ TEXT("widget_path"), TEXT("string"), TEXT("Path to the Widget Blueprint"), true },
+			{ TEXT("button_name"), TEXT("string"), TEXT("Name for the Button"), true },
+			{ TEXT("text"), TEXT("string"), TEXT("Button text"), false, TEXT("Button") },
+			{ TEXT("position"), TEXT("object"), TEXT("Position in canvas panel {x, y}"), false },
+			{ TEXT("size"), TEXT("object"), TEXT("Size of the button {width, height}"), false },
+			{ TEXT("font_size"), TEXT("number"), TEXT("Font size for button text"), false, TEXT("12") }
+		};
+	}
+	
+	virtual FECACommandResult Execute(const TSharedPtr<FJsonObject>& Params) override;
+};
+
+/**
+ * Add a widget instance to the game viewport
+ */
+class FECACommand_AddWidgetToViewport : public IECACommand
+{
+public:
+	virtual FString GetName() const override { return TEXT("add_widget_to_viewport"); }
+	virtual FString GetDescription() const override { return TEXT("Add an instance of a Widget Blueprint to the game viewport"); }
+	virtual FString GetCategory() const override { return TEXT("UMG"); }
+	
+	virtual TArray<FECACommandParam> GetParameters() const override
+	{
+		return {
+			{ TEXT("widget_path"), TEXT("string"), TEXT("Path to the Widget Blueprint"), true },
+			{ TEXT("z_order"), TEXT("number"), TEXT("Z-order for widget display"), false, TEXT("0") }
+		};
+	}
+	
+	virtual FECACommandResult Execute(const TSharedPtr<FJsonObject>& Params) override;
+};
+
+/**
+ * Bind an event to a widget (e.g., button click)
+ */
+class FECACommand_BindWidgetEvent : public IECACommand
+{
+public:
+	virtual FString GetName() const override { return TEXT("bind_widget_event"); }
+	virtual FString GetDescription() const override { return TEXT("Bind an event to a widget (e.g., OnClicked for buttons)"); }
+	virtual FString GetCategory() const override { return TEXT("UMG"); }
+	
+	virtual TArray<FECACommandParam> GetParameters() const override
+	{
+		return {
+			{ TEXT("widget_path"), TEXT("string"), TEXT("Path to the Widget Blueprint"), true },
+			{ TEXT("widget_name"), TEXT("string"), TEXT("Name of the widget to bind"), true },
+			{ TEXT("event_name"), TEXT("string"), TEXT("Name of the event (OnClicked, OnHovered, etc.)"), true },
+			{ TEXT("function_name"), TEXT("string"), TEXT("Name of the function to call"), false }
+		};
+	}
+	
+	virtual FECACommandResult Execute(const TSharedPtr<FJsonObject>& Params) override;
+};
+
+/**
+ * Set up text block binding for dynamic updates
+ */
+class FECACommand_SetTextBlockBinding : public IECACommand
+{
+public:
+	virtual FString GetName() const override { return TEXT("set_text_block_binding"); }
+	virtual FString GetDescription() const override { return TEXT("Set up a text binding for dynamic text updates"); }
+	virtual FString GetCategory() const override { return TEXT("UMG"); }
+	
+	virtual TArray<FECACommandParam> GetParameters() const override
+	{
+		return {
+			{ TEXT("widget_path"), TEXT("string"), TEXT("Path to the Widget Blueprint"), true },
+			{ TEXT("text_block_name"), TEXT("string"), TEXT("Name of the Text Block widget"), true },
+			{ TEXT("binding_function"), TEXT("string"), TEXT("Name of the function that returns the text"), true }
+		};
+	}
+	
+	virtual FECACommandResult Execute(const TSharedPtr<FJsonObject>& Params) override;
+};
+
+/**
+ * Add an Image widget to a Widget Blueprint
+ */
+class FECACommand_AddImageToWidget : public IECACommand
+{
+public:
+	virtual FString GetName() const override { return TEXT("add_image_to_widget"); }
+	virtual FString GetDescription() const override { return TEXT("Add an Image widget to a Widget Blueprint"); }
+	virtual FString GetCategory() const override { return TEXT("UMG"); }
+	
+	virtual TArray<FECACommandParam> GetParameters() const override
+	{
+		return {
+			{ TEXT("widget_path"), TEXT("string"), TEXT("Path to the Widget Blueprint"), true },
+			{ TEXT("image_name"), TEXT("string"), TEXT("Name for the Image widget"), true },
+			{ TEXT("texture_path"), TEXT("string"), TEXT("Path to the texture asset"), false },
+			{ TEXT("position"), TEXT("object"), TEXT("Position in canvas panel {x, y}"), false },
+			{ TEXT("size"), TEXT("object"), TEXT("Size of the image {width, height}"), false },
+			{ TEXT("color_and_opacity"), TEXT("object"), TEXT("Color tint {r, g, b, a}"), false }
+		};
+	}
+	
+	virtual FECACommandResult Execute(const TSharedPtr<FJsonObject>& Params) override;
+};
+
+/**
+ * Get information about a Widget Blueprint
+ */
+class FECACommand_GetWidgetInfo : public IECACommand
+{
+public:
+	virtual FString GetName() const override { return TEXT("get_widget_info"); }
+	virtual FString GetDescription() const override { return TEXT("Get information about a Widget Blueprint"); }
+	virtual FString GetCategory() const override { return TEXT("UMG"); }
+	
+	virtual TArray<FECACommandParam> GetParameters() const override
+	{
+		return {
+			{ TEXT("widget_path"), TEXT("string"), TEXT("Path to the Widget Blueprint"), true }
+		};
+	}
+	
+	virtual FECACommandResult Execute(const TSharedPtr<FJsonObject>& Params) override;
+};
diff --git a/Engine/Plugins/Experimental/ECABridge/Source/ECABridge/Public/Commands/ECAViewCommands.h b/Engine/Plugins/Experimental/ECABridge/Source/ECABridge/Public/Commands/ECAViewCommands.h
new file mode 100644
--- /dev/null
+++ b/Engine/Plugins/Experimental/ECABridge/Source/ECABridge/Public/Commands/ECAViewCommands.h
@@ -0,0 +1,85 @@
+// Copyright Epic Games, Inc. All Rights Reserved.
+
+#pragma once
+
+#include "Commands/ECACommand.h"
+
+/**
+ * Get actors visible within a camera's view frustum.
+ * 
+ * This command allows the AI to "see" what a camera (or the editor viewport) can see,
+ * returning information about objects within the view frustum up to a specified distance.
+ */
+class FECACommand_GetActorsInView : public IECACommand
+{
+public:
+	virtual FString GetName() const override { return TEXT("get_actors_in_view"); }
+	virtual FString GetDescription() const override { return TEXT("Get actors visible within a camera's view frustum. Returns near objects that the camera can see, useful for AI scene understanding."); }
+	virtual FString GetCategory() const override { return TEXT("View"); }
+	
+	virtual TArray<FECACommandParam> GetParameters() const override
+	{
+		return {
+			{ TEXT("camera_name"), TEXT("string"), TEXT("Name of a CameraActor to use as the view source. If not specified, uses the active editor viewport camera."), false },
+			{ TEXT("max_distance"), TEXT("number"), TEXT("Maximum distance from camera to include actors (in centimeters). Default: 10000 (100 meters)"), false, TEXT("10000") },
+			{ TEXT("min_distance"), TEXT("number"), TEXT("Minimum distance from camera to include actors (in centimeters). Default: 0"), false, TEXT("0") },
+			{ TEXT("include_hidden"), TEXT("boolean"), TEXT("Include hidden actors in results. Default: false"), false, TEXT("false") },
+			{ TEXT("class_filter"), TEXT("string"), TEXT("Filter actors by class name (partial match)"), false },
+			{ TEXT("include_components"), TEXT("boolean"), TEXT("Include component details for each actor. Default: false"), false, TEXT("false") },
+			{ TEXT("max_results"), TEXT("number"), TEXT("Maximum number of actors to return, sorted by distance. Default: 100"), false, TEXT("100") },
+			{ TEXT("fov_scale"), TEXT("number"), TEXT("Scale factor for the field of view (1.0 = camera's actual FOV, 2.0 = double width). Default: 1.0"), false, TEXT("1.0") }
+		};
+	}
+	
+	virtual FECACommandResult Execute(const TSharedPtr<FJsonObject>& Params) override;
+};
+
+/**
+ * Get the current editor viewport camera properties.
+ * 
+ * Returns information about the active editor viewport camera including
+ * position, rotation, FOV, and other view settings.
+ */
+class FECACommand_GetViewportCamera : public IECACommand
+{
+public:
+	virtual FString GetName() const override { return TEXT("get_viewport_camera"); }
+	virtual FString GetDescription() const override { return TEXT("Get the current editor viewport camera position, rotation, and properties."); }
+	virtual FString GetCategory() const override { return TEXT("View"); }
+	
+	virtual TArray<FECACommandParam> GetParameters() const override
+	{
+		return {
+			{ TEXT("viewport_index"), TEXT("number"), TEXT("Index of the viewport to query. Default: 0 (active viewport)"), false, TEXT("0") }
+		};
+	}
+	
+	virtual FECACommandResult Execute(const TSharedPtr<FJsonObject>& Params) override;
+};
+
+/**
+ * Describe what a camera can see in natural language.
+ * 
+ * Returns a structured description of the scene visible from a camera's perspective,
+ * organized by distance zones and object types.
+ */
+class FECACommand_DescribeView : public IECACommand
+{
+public:
+	virtual FString GetName() const override { return TEXT("describe_view"); }
+	virtual FString GetDescription() const override { return TEXT("Get a structured description of what a camera can see, organized by distance zones (near/mid/far) and object categories."); }
+	virtual FString GetCategory() const override { return TEXT("View"); }
+	
+	virtual TArray<FECACommandParam> GetParameters() const override
+	{
+		return {
+			{ TEXT("camera_name"), TEXT("string"), TEXT("Name of a CameraActor to use. If not specified, uses the active editor viewport camera."), false },
+			{ TEXT("max_distance"), TEXT("number"), TEXT("Maximum view distance in centimeters. Default: 50000 (500 meters)"), false, TEXT("50000") },
+			{ TEXT("near_threshold"), TEXT("number"), TEXT("Distance threshold for 'near' zone in cm. Default: 500 (5 meters)"), false, TEXT("500") },
+			{ TEXT("mid_threshold"), TEXT("number"), TEXT("Distance threshold for 'mid' zone in cm. Default: 2000 (20 meters)"), false, TEXT("2000") },
+			{ TEXT("include_hidden"), TEXT("boolean"), TEXT("Include hidden actors. Default: false"), false, TEXT("false") }
+		};
+	}
+	
+	virtual FECACommandResult Execute(const TSharedPtr<FJsonObject>& Params) override;
+};
diff --git a/Engine/Plugins/Experimental/ECABridge/Source/ECABridge/Public/Commands/ECAWidgetTreeCommands.h b/Engine/Plugins/Experimental/ECABridge/Source/ECABridge/Public/Commands/ECAWidgetTreeCommands.h
new file mode 100644
--- /dev/null
+++ b/Engine/Plugins/Experimental/ECABridge/Source/ECABridge/Public/Commands/ECAWidgetTreeCommands.h
@@ -0,0 +1,195 @@
+// Copyright Epic Games, Inc. All Rights Reserved.
+
+#pragma once
+
+#include "Commands/ECACommand.h"
+
+/**
+ * Wrap a widget with a container widget (e.g., insert a SizeBox between a widget and its parent).
+ * Follows the same algorithm as FWidgetBlueprintEditorUtils::WrapWidgets.
+ */
+class FECACommand_WrapWidget : public IECACommand
+{
+public:
+	virtual FString GetName() const override { return TEXT("wrap_widget"); }
+	virtual FString GetDescription() const override { return TEXT("Wrap a widget with a container widget (e.g., insert a SizeBox between a widget and its parent)"); }
+	virtual FString GetCategory() const override { return TEXT("UMG"); }
+
+	virtual TArray<FECACommandParam> GetParameters() const override
+	{
+		return {
+			{ TEXT("widget_path"), TEXT("string"), TEXT("Path to the Widget Blueprint"), true },
+			{ TEXT("widget_name"), TEXT("string"), TEXT("Name of the widget to wrap"), true },
+			{ TEXT("wrapper_class"), TEXT("string"), TEXT("Class name of the wrapper widget (e.g., SizeBox, Overlay, Border, ScaleBox)"), true },
+			{ TEXT("wrapper_name"), TEXT("string"), TEXT("Name for the new wrapper widget (auto-generated if omitted)"), false },
+		};
+	}
+
+	virtual FECACommandResult Execute(const TSharedPtr<FJsonObject>& Params) override;
+};
+
+/**
+ * Unwrap a container widget by removing it and promoting its single child to take its place.
+ * Follows the same algorithm as FWidgetBlueprintEditorUtils::ReplaceWidgetWithChildren.
+ */
+class FECACommand_UnwrapWidget : public IECACommand
+{
+public:
+	virtual FString GetName() const override { return TEXT("unwrap_widget"); }
+	virtual FString GetDescription() const override { return TEXT("Remove a wrapper widget and promote its single child to take its place in the hierarchy"); }
+	virtual FString GetCategory() const override { return TEXT("UMG"); }
+
+	virtual TArray<FECACommandParam> GetParameters() const override
+	{
+		return {
+			{ TEXT("widget_path"), TEXT("string"), TEXT("Path to the Widget Blueprint"), true },
+			{ TEXT("widget_name"), TEXT("string"), TEXT("Name of the wrapper widget to remove"), true },
+		};
+	}
+
+	virtual FECACommandResult Execute(const TSharedPtr<FJsonObject>& Params) override;
+};
+
+/**
+ * Move a widget from its current parent to a different parent widget.
+ * Note: Slot properties (padding, alignment, size rule) are lost on reparent since a new slot is created.
+ */
+class FECACommand_ReparentWidget : public IECACommand
+{
+public:
+	virtual FString GetName() const override { return TEXT("reparent_widget"); }
+	virtual FString GetDescription() const override { return TEXT("Move a widget from its current parent to a different parent widget. Slot properties are lost on reparent."); }
+	virtual FString GetCategory() const override { return TEXT("UMG"); }
+
+	virtual TArray<FECACommandParam> GetParameters() const override
+	{
+		return {
+			{ TEXT("widget_path"), TEXT("string"), TEXT("Path to the Widget Blueprint"), true },
+			{ TEXT("widget_name"), TEXT("string"), TEXT("Name of the widget to move"), true },
+			{ TEXT("new_parent_name"), TEXT("string"), TEXT("Name of the target parent widget"), true },
+			{ TEXT("insert_index"), TEXT("number"), TEXT("Position in new parent (-1 or omit to append at end)"), false, TEXT("-1") },
+		};
+	}
+
+	virtual FECACommandResult Execute(const TSharedPtr<FJsonObject>& Params) override;
+};
+
+/**
+ * Set layout properties on a widget's slot (padding, alignment, size rule).
+ * Slot type is auto-detected and properties are applied accordingly.
+ */
+class FECACommand_SetSlotProperty : public IECACommand
+{
+public:
+	virtual FString GetName() const override { return TEXT("set_slot_property"); }
+	virtual FString GetDescription() const override { return TEXT("Set layout properties on a widget's slot (padding, alignment, size rule). Slot type is auto-detected."); }
+	virtual FString GetCategory() const override { return TEXT("UMG"); }
+
+	virtual TArray<FECACommandParam> GetParameters() const override
+	{
+		return {
+			{ TEXT("widget_path"), TEXT("string"), TEXT("Path to the Widget Blueprint"), true },
+			{ TEXT("widget_name"), TEXT("string"), TEXT("Name of the child widget whose slot to modify"), true },
+			{ TEXT("padding"), TEXT("object"), TEXT("Padding as {left, top, right, bottom}"), false },
+			{ TEXT("horizontal_alignment"), TEXT("string"), TEXT("Horizontal alignment: Fill, Left, Center, Right"), false },
+			{ TEXT("vertical_alignment"), TEXT("string"), TEXT("Vertical alignment: Fill, Top, Center, Bottom"), false },
+			{ TEXT("size_rule"), TEXT("string"), TEXT("Size rule: Auto or Fill (only for VerticalBox/HorizontalBox slots)"), false },
+			{ TEXT("size_value"), TEXT("number"), TEXT("Fill weight when size_rule is Fill (default 1.0)"), false, TEXT("1.0") },
+		};
+	}
+
+	virtual FECACommandResult Execute(const TSharedPtr<FJsonObject>& Params) override;
+};
+
+/**
+ * Set properties directly on a widget (not its slot).
+ * Supports common widget properties like Visibility, RenderOpacity,
+ * and type-specific properties like SizeBox overrides.
+ */
+class FECACommand_SetWidgetProperty : public IECACommand
+{
+public:
+	virtual FString GetName() const override { return TEXT("set_widget_property"); }
+	virtual FString GetDescription() const override { return TEXT("Set properties on a widget (Visibility, SizeBox overrides, Text, etc.). Use set_slot_property for layout/padding."); }
+	virtual FString GetCategory() const override { return TEXT("UMG"); }
+
+	virtual TArray<FECACommandParam> GetParameters() const override
+	{
+		return {
+			{ TEXT("widget_path"), TEXT("string"), TEXT("Path to the Widget Blueprint"), true },
+			{ TEXT("widget_name"), TEXT("string"), TEXT("Name of the widget to modify"), true },
+			{ TEXT("properties"), TEXT("object"), TEXT("Key-value pairs of properties to set (e.g., {\"WidthOverride\": 400, \"Visibility\": \"Collapsed\"})"), true },
+		};
+	}
+
+	virtual FECACommandResult Execute(const TSharedPtr<FJsonObject>& Params) override;
+};
+
+/**
+ * Add a child widget of any class to a parent widget in a Widget Blueprint.
+ * Works with any UWidget subclass — engine types (TextBlock, Image, SizeBox) and
+ * project types (CommonTextBlock, WBP_MyWidget_C, etc.).
+ */
+class FECACommand_AddChildWidget : public IECACommand
+{
+public:
+	virtual FString GetName() const override { return TEXT("add_child_widget"); }
+	virtual FString GetDescription() const override { return TEXT("Add a child widget of any class to a parent widget. Works with engine and project widget types."); }
+	virtual FString GetCategory() const override { return TEXT("UMG"); }
+
+	virtual TArray<FECACommandParam> GetParameters() const override
+	{
+		return {
+			{ TEXT("widget_path"), TEXT("string"), TEXT("Path to the Widget Blueprint"), true },
+			{ TEXT("parent_name"), TEXT("string"), TEXT("Name of the parent widget to add the child to"), true },
+			{ TEXT("widget_class"), TEXT("string"), TEXT("Class name of the widget to create (e.g., TextBlock, CommonTextBlock, SizeBox, WBP_MyWidget_C)"), true },
+			{ TEXT("widget_name"), TEXT("string"), TEXT("Name for the new widget"), true },
+			{ TEXT("insert_index"), TEXT("number"), TEXT("Position in parent (-1 or omit to append at end)"), false, TEXT("-1") },
+			{ TEXT("skip_compile"), TEXT("boolean"), TEXT("Skip blueprint recompile (batch multiple adds, then call compile_widget_blueprint)"), false },
+		};
+	}
+
+	virtual FECACommandResult Execute(const TSharedPtr<FJsonObject>& Params) override;
+};
+
+/**
+ * Compile a Widget Blueprint. Use after batching operations with skip_compile=true.
+ */
+class FECACommand_CompileWidgetBlueprint : public IECACommand
+{
+public:
+	virtual FString GetName() const override { return TEXT("compile_widget_blueprint"); }
+	virtual FString GetDescription() const override { return TEXT("Compile a Widget Blueprint. Use after batching add_child_widget calls with skip_compile=true."); }
+	virtual FString GetCategory() const override { return TEXT("UMG"); }
+
+	virtual TArray<FECACommandParam> GetParameters() const override
+	{
+		return {
+			{ TEXT("widget_path"), TEXT("string"), TEXT("Path to the Widget Blueprint"), true },
+		};
+	}
+
+	virtual FECACommandResult Execute(const TSharedPtr<FJsonObject>& Params) override;
+};
+
+/**
+ * Remove a widget and all its descendants from a Widget Blueprint.
+ * The widget is detached from its parent and fully removed from the WidgetTree.
+ */
+class FECACommand_RemoveWidget : public IECACommand
+{
+public:
+	virtual FString GetName() const override { return TEXT("remove_widget"); }
+	virtual FString GetDescription() const override { return TEXT("Remove a widget and all its descendants from a Widget Blueprint"); }
+	virtual FString GetCategory() const override { return TEXT("UMG"); }
+
+	virtual TArray<FECACommandParam> GetParameters() const override
+	{
+		return {
+			{ TEXT("widget_path"), TEXT("string"), TEXT("Path to the Widget Blueprint"), true },
+			{ TEXT("widget_name"), TEXT("string"), TEXT("Name of the widget to remove"), true },
+		};
+	}
+
+	virtual FECACommandResult Execute(const TSharedPtr<FJsonObject>& Params) override;
+};
diff --git a/Engine/Plugins/Experimental/ECABridge/Source/ECABridge/Public/ECABridge.h b/Engine/Plugins/Experimental/ECABridge/Source/ECABridge/Public/ECABridge.h
new file mode 100644
--- /dev/null
+++ b/Engine/Plugins/Experimental/ECABridge/Source/ECABridge/Public/ECABridge.h
@@ -0,0 +1,54 @@
+// Copyright Epic Games, Inc. All Rights Reserved.
+
+#pragma once
+
+#include "CoreMinimal.h"
+#include "UObject/NoExportTypes.h"
+#include "Dom/JsonObject.h"
+#include "ECABridge.generated.h"
+
+class FECAMCPServer;
+
+/**
+ * ECA Bridge
+ * 
+ * Main coordinator for Epic Code Assistant integration.
+ * Manages the MCP HTTP server and routes commands to the appropriate handlers.
+ */
+UCLASS()
+class ECABRIDGE_API UECABridge : public UObject
+{
+	GENERATED_BODY()
+
+public:
+	/** Initialize the bridge and start the MCP server */
+	void Initialize();
+	
+	/** Shutdown the bridge and stop the MCP server */
+	void Shutdown();
+	
+	/** Process a JSON command string and return the response */
+	FString ProcessCommand(const FString& JsonCommand);
+	
+	/** Check if the server is running */
+	bool IsRunning() const;
+	
+	/** Get the port the MCP server is listening on */
+	int32 GetPort() const;
+	
+	/** Get the number of commands processed */
+	int32 GetCommandsProcessed() const { return CommandsProcessed; }
+	
+	/** Log all registered commands to the output log */
+	void LogRegisteredCommands() const;
+
+private:
+	/** The MCP HTTP server */
+	TSharedPtr<FECAMCPServer> MCPServer;
+	
+	/** Counter for processed commands */
+	TAtomic<int32> CommandsProcessed{0};
+	
+	/** Parse JSON command and extract command name and params */
+	bool ParseCommand(const FString& JsonCommand, FString& OutCommandName, TSharedPtr<FJsonObject>& OutParams);
+};
diff --git a/Engine/Plugins/Experimental/ECABridge/Source/ECABridge/Public/ECABridgeModule.h b/Engine/Plugins/Experimental/ECABridge/Source/ECABridge/Public/ECABridgeModule.h
new file mode 100644
--- /dev/null
+++ b/Engine/Plugins/Experimental/ECABridge/Source/ECABridge/Public/ECABridgeModule.h
@@ -0,0 +1,143 @@
+// Copyright Epic Games, Inc. All Rights Reserved.
+
+#pragma once
+
+#include "CoreMinimal.h"
+#include "Modules/ModuleManager.h"
+#include "Engine/EngineTypes.h"  // For FTimerHandle
+#include "Editor/UnrealEdTypes.h"  // For ELevelViewportType
+
+class UECABridge;
+class UFactory;
+struct FAssetData;
+class AActor;
+
+/**
+ * ECA Bridge Module
+ * 
+ * Provides TCP server for Epic Code Assistant to control Unreal Editor.
+ * Commands are executed on the game thread to safely interact with editor subsystems.
+ * Also monitors editor events and queues them for the AI to process.
+ */
+class FECABridgeModule : public IModuleInterface
+{
+public:
+	/** IModuleInterface implementation */
+	virtual void StartupModule() override;
+	virtual void ShutdownModule() override;
+	
+	/** Get the bridge instance */
+	UECABridge* GetBridge() const { return Bridge; }
+	
+	/** Module singleton access */
+	static FECABridgeModule& Get()
+	{
+		return FModuleManager::GetModuleChecked<FECABridgeModule>("ECABridge");
+	}
+	
+	static bool IsAvailable()
+	{
+		return FModuleManager::Get().IsModuleLoaded("ECABridge");
+	}
+
+private:
+	/** The bridge instance that handles commands */
+	UECABridge* Bridge = nullptr;
+	
+	/** Register toolbar extension */
+	void RegisterToolbarExtension();
+	
+	/** Console commands for debugging */
+	void RegisterConsoleCommands();
+	
+	/** Register editor event handlers for pending events */
+	void RegisterEditorEventHandlers();
+	
+	/** Unregister editor event handlers */
+	void UnregisterEditorEventHandlers();
+	
+	/** Called when the editor finishes initializing (for deferred handler registration) */
+	void OnEditorInitialized(double Duration);
+	
+	/** Register handlers that depend on GEditor being valid */
+	void RegisterGEditorDependentHandlers();
+	
+	//--------------------------------------------------------------------------
+	// Event Callbacks
+	//--------------------------------------------------------------------------
+	
+	/** Called when editor selection changes */
+	void OnEditorSelectionChanged(UObject* SelectionObject);
+	
+	/** Called when an asset is imported via the import dialog */
+	void OnAssetImported(UFactory* Factory, UObject* CreatedObject);
+	
+	/** Called when a new asset is added to the asset registry */
+	void OnAssetAdded(const FAssetData& AssetData);
+	
+	/** Called when an asset is removed from the asset registry */
+	void OnAssetRemoved(const FAssetData& AssetData);
+	
+	/** Called when an asset is renamed */
+	void OnAssetRenamed(const FAssetData& AssetData, const FString& OldObjectPath);
+	
+	/** Called when any Blueprint is compiled */
+	void OnBlueprintCompiled();
+	
+	/** Called when a map/level is opened */
+	void OnMapOpened(const FString& Filename, bool bAsTemplate);
+	
+	/** Called when Play In Editor starts */
+	void OnPIEStarted(bool bIsSimulating);
+	
+	/** Called when Play In Editor ends */
+	void OnPIEEnded(bool bIsSimulating);
+	
+	/** Called when an actor is spawned in the editor */
+	void OnActorSpawned(AActor* Actor);
+	
+	/** Called when an actor is deleted in the editor */
+	void OnActorDeleted(AActor* Actor);
+	
+	/** Called when the editor camera moves */
+	void OnEditorCameraMoved(const FVector& Location, const FRotator& Rotation, ELevelViewportType ViewportType, int32 ViewIndex);
+	
+	/** Timer callback to send describe_view after camera stops moving */
+	void OnCameraStoppedMoving();
+	
+	//--------------------------------------------------------------------------
+	// Delegate Handles
+	//--------------------------------------------------------------------------
+	
+	FDelegateHandle SelectionChangedHandle;
+	FDelegateHandle AssetImportHandle;
+	FDelegateHandle AssetAddedHandle;
+	FDelegateHandle AssetRemovedHandle;
+	FDelegateHandle AssetRenamedHandle;
+	FDelegateHandle BlueprintCompiledHandle;
+	FDelegateHandle MapChangedHandle;
+	FDelegateHandle PIEStartedHandle;
+	FDelegateHandle PIEEndedHandle;
+	FDelegateHandle ActorSpawnedHandle;
+	FDelegateHandle ActorDeletedHandle;
+	FDelegateHandle CameraMovedHandle;
+	
+	//--------------------------------------------------------------------------
+	// Camera Movement Tracking
+	//--------------------------------------------------------------------------
+	
+	/** Timer handle for debouncing camera movement events */
+	FTimerHandle CameraMovedTimerHandle;
+	
+	/** Last camera location when movement started */
+	FVector LastCameraLocation;
+	
+	/** Last camera rotation when movement started */
+	FRotator LastCameraRotation;
+	
+	/** Whether the camera is currently moving */
+	bool bCameraIsMoving = false;
+	
+	/** Debounce time in seconds - wait this long after camera stops before sending event */
+	static constexpr float CameraStoppedDebounceTime = 0.5f;
+};
diff --git a/Engine/Plugins/Experimental/ECABridge/Source/ECABridge/Public/ECABridgeSettings.h b/Engine/Plugins/Experimental/ECABridge/Source/ECABridge/Public/ECABridgeSettings.h
new file mode 100644
--- /dev/null
+++ b/Engine/Plugins/Experimental/ECABridge/Source/ECABridge/Public/ECABridgeSettings.h
@@ -0,0 +1,57 @@
+// Copyright Epic Games, Inc. All Rights Reserved.
+
+#pragma once
+
+#include "CoreMinimal.h"
+#include "Engine/DeveloperSettings.h"
+#include "ECABridgeSettings.generated.h"
+
+/**
+ * ECA Bridge Settings
+ * 
+ * Configure the MCP server that allows AI assistants (Claude Desktop, Cursor, etc.)
+ * to control Unreal Editor.
+ */
+UCLASS(config=EditorPerProjectUserSettings, meta=(DisplayName="ECA Bridge"))
+class ECABRIDGE_API UECABridgeSettings : public UDeveloperSettings
+{
+	GENERATED_BODY()
+
+public:
+	UECABridgeSettings();
+
+	/** Port for the MCP HTTP server (for Claude Desktop, Cursor, etc.) */
+	UPROPERTY(config, EditAnywhere, Category="Server", meta=(ClampMin=1024, ClampMax=65535))
+	int32 ServerPort = 3000;
+	
+	/** Whether to start the server automatically when the editor launches */
+	UPROPERTY(config, EditAnywhere, Category="Server")
+	bool bAutoStart = true;
+	
+	/** Maximum time (in seconds) to wait for a command to complete */
+	UPROPERTY(config, EditAnywhere, Category="Server", meta=(ClampMin=1, ClampMax=300))
+	int32 CommandTimeoutSeconds = 30;
+	
+	/** Enable verbose logging for debugging */
+	UPROPERTY(config, EditAnywhere, Category="Debugging")
+	bool bVerboseLogging = false;
+	
+	/** Log all incoming commands (may contain sensitive data) */
+	UPROPERTY(config, EditAnywhere, Category="Debugging")
+	bool bLogCommands = false;
+	
+	/** Allow connections from external IPs (security risk - use with caution) */
+	UPROPERTY(config, EditAnywhere, Category="Security", AdvancedDisplay)
+	bool bAllowRemoteConnections = false;
+	
+	/** Get the singleton settings object */
+	static const UECABridgeSettings* Get()
+	{
+		return GetDefault<UECABridgeSettings>();
+	}
+	
+	// UDeveloperSettings interface
+	virtual FName GetCategoryName() const override { return FName("Plugins"); }
+	virtual FText GetSectionText() const override { return NSLOCTEXT("ECABridge", "SettingsSection", "ECA Bridge"); }
+	virtual FText GetSectionDescription() const override { return NSLOCTEXT("ECABridge", "SettingsDescription", "Configure Epic Code Assistant MCP server for AI-assisted development"); }
+};
diff --git a/Engine/Plugins/Experimental/ECABridge/Source/ECABridge/Public/ECAEventQueue.h b/Engine/Plugins/Experimental/ECABridge/Source/ECABridge/Public/ECAEventQueue.h
new file mode 100644
--- /dev/null
+++ b/Engine/Plugins/Experimental/ECABridge/Source/ECABridge/Public/ECAEventQueue.h
@@ -0,0 +1,80 @@
+// Copyright Epic Games, Inc. All Rights Reserved.
+
+#pragma once
+
+#include "CoreMinimal.h"
+#include "Dom/JsonObject.h"
+
+/**
+ * ECA Event Queue Helper Functions
+ * 
+ * Use these functions to queue events that will be polled by ECA.
+ * Events appear in the AI conversation and can optionally trigger AI responses.
+ * 
+ * Example usage:
+ *   #include "ECAEventQueue.h"
+ *   
+ *   // When user selects something in the editor:
+ *   ECAEvents::QueueSelectionChanged("BP_PlayerCharacter", "/Game/Blueprints/BP_PlayerCharacter");
+ *   
+ *   // When an asset is imported:
+ *   ECAEvents::QueueAssetImported("/Game/Meshes/SM_Chair", "StaticMesh");
+ *   
+ *   // For custom events:
+ *   TSharedPtr<FJsonObject> Data = MakeShared<FJsonObject>();
+ *   Data->SetStringField("my_field", "my_value");
+ *   ECAEvents::QueueCustomEvent("my_event_type", "Something happened!", Data, true);
+ */
+namespace ECAEvents
+{
+	/**
+	 * Queue a selection changed event.
+	 * Triggers AI response by default (user explicitly selected something).
+	 * 
+	 * @param ActorName  Display name of the selected actor
+	 * @param ActorPath  Full path to the actor
+	 */
+	ECABRIDGE_API void QueueSelectionChanged(const FString& ActorName, const FString& ActorPath);
+
+	/**
+	 * Queue an asset imported event.
+	 * Does NOT trigger AI response by default (informational).
+	 * 
+	 * @param AssetPath  Path to the imported asset
+	 * @param AssetType  Type of asset (e.g., "StaticMesh", "Texture2D")
+	 */
+	ECABRIDGE_API void QueueAssetImported(const FString& AssetPath, const FString& AssetType);
+
+	/**
+	 * Queue an asset saved event.
+	 * Does NOT trigger AI response by default (informational).
+	 * 
+	 * @param AssetPath  Path to the saved asset
+	 */
+	ECABRIDGE_API void QueueAssetSaved(const FString& AssetPath);
+
+	/**
+	 * Queue a Blueprint compilation event.
+	 * Triggers AI response only on failure (to help debug).
+	 * 
+	 * @param BlueprintPath  Path to the Blueprint
+	 * @param bSuccess       Whether compilation succeeded
+	 * @param ErrorMessage   Error message if compilation failed
+	 */
+	ECABRIDGE_API void QueueBlueprintCompiled(const FString& BlueprintPath, bool bSuccess, const FString& ErrorMessage = TEXT(""));
+
+	/**
+	 * Queue a custom event with full control over all parameters.
+	 * 
+	 * @param EventType         Type identifier for the event
+	 * @param Message           Human-readable description
+	 * @param Data              Optional JSON data object
+	 * @param bTriggerAIResponse  Whether ECA should respond to this event
+	 */
+	ECABRIDGE_API void QueueCustomEvent(
+		const FString& EventType, 
+		const FString& Message, 
+		TSharedPtr<FJsonObject> Data = nullptr, 
+		bool bTriggerAIResponse = true
+	);
+}
diff --git a/Engine/Plugins/Experimental/ECABridge/Source/ECABridge/Public/ECAMCPServer.h b/Engine/Plugins/Experimental/ECABridge/Source/ECABridge/Public/ECAMCPServer.h
new file mode 100644
--- /dev/null
+++ b/Engine/Plugins/Experimental/ECABridge/Source/ECABridge/Public/ECAMCPServer.h
@@ -0,0 +1,97 @@
+// Copyright Epic Games, Inc. All Rights Reserved.
+
+#pragma once
+
+#include "CoreMinimal.h"
+#include "Dom/JsonObject.h"
+#include "Dom/JsonValue.h"
+#include "HttpServerModule.h"
+#include "IHttpRouter.h"
+#include "HttpServerRequest.h"
+#include "HttpServerResponse.h"
+#include "HttpResultCallback.h"
+
+class UECABridge;
+
+/**
+ * MCP Streamable HTTP Server
+ * 
+ * Implements the MCP Streamable HTTP transport (2025-03-26 spec).
+ * This is a simpler transport that uses standard HTTP request/response patterns.
+ * 
+ * Endpoints:
+ *   POST /mcp    - Main MCP endpoint for all JSON-RPC messages
+ *   GET  /mcp    - Optional SSE endpoint for server-initiated messages (not implemented)
+ *   GET  /health - Health check
+ */
+class ECABRIDGE_API FECAMCPServer
+{
+public:
+	FECAMCPServer(UECABridge* InBridge);
+	~FECAMCPServer();
+
+	/** Start the HTTP server */
+	bool Start(int32 Port = 3000);
+	
+	/** Stop the HTTP server */
+	void Stop();
+	
+	/** Check if server is running */
+	bool IsRunning() const { return bIsRunning; }
+	
+	/** Get the port */
+	int32 GetPort() const { return ServerPort; }
+
+private:
+	/** Handle MCP POST request (JSON-RPC messages) */
+	bool HandleMCPPost(const FHttpServerRequest& Request, const FHttpResultCallback& OnComplete);
+	
+	/** Handle health check */
+	bool HandleHealth(const FHttpServerRequest& Request, const FHttpResultCallback& OnComplete);
+	
+	/** Handle CORS preflight (OPTIONS) requests */
+	bool HandleCORSPreflight(const FHttpServerRequest& Request, const FHttpResultCallback& OnComplete);
+
+	/** Process a single JSON-RPC request and return the response */
+	TSharedPtr<FJsonObject> ProcessJsonRpcRequest(const TSharedPtr<FJsonObject>& Request);
+	
+	/** Handle MCP initialize request */
+	TSharedPtr<FJsonObject> HandleInitialize(const TSharedPtr<FJsonObject>& Params);
+	
+	/** Handle tools/list request */
+	TSharedPtr<FJsonObject> HandleToolsList();
+	
+	/** Handle tools/call request */
+	TSharedPtr<FJsonObject> HandleToolsCall(const TSharedPtr<FJsonObject>& Params);
+	
+	/** Build MCP tool definitions from registered commands */
+	TArray<TSharedPtr<FJsonValue>> BuildToolDefinitions();
+	
+	/** Create JSON-RPC response wrapper */
+	TSharedPtr<FJsonObject> CreateJsonRpcResponse(const TSharedPtr<FJsonObject>& Result, const TSharedPtr<FJsonValue>& RequestId);
+	
+	/** Create JSON-RPC error response */
+	TSharedPtr<FJsonObject> CreateJsonRpcError(int32 Code, const FString& Message, const TSharedPtr<FJsonValue>& RequestId);
+	
+	/** Serialize JSON object to string */
+	FString SerializeJson(const TSharedPtr<FJsonObject>& JsonObject);
+
+	/** The bridge to execute commands */
+	UECABridge* Bridge;
+	
+	/** HTTP router */
+	TSharedPtr<IHttpRouter> HttpRouter;
+	
+	/** Route handles for cleanup */
+	FHttpRouteHandle MCPRouteHandle;
+	FHttpRouteHandle HealthRouteHandle;
+	
+	/** Server state */
+	bool bIsRunning = false;
+	int32 ServerPort = 3000;
+	
+	/** MCP protocol info */
+	const FString ServerName = TEXT("eca-unreal-editor");
+	const FString ServerVersion = TEXT("1.0.0");
+	const FString ProtocolVersion = TEXT("2025-03-26");
+};
diff --git a/Engine/Plugins/Experimental/ECABridge/Source/ECABridge/Public/ECAServerRunnable.h b/Engine/Plugins/Experimental/ECABridge/Source/ECABridge/Public/ECAServerRunnable.h
new file mode 100644
--- /dev/null
+++ b/Engine/Plugins/Experimental/ECABridge/Source/ECABridge/Public/ECAServerRunnable.h
@@ -0,0 +1,52 @@
+// Copyright Epic Games, Inc. All Rights Reserved.
+
+#pragma once
+
+#include "CoreMinimal.h"
+#include "HAL/Runnable.h"
+
+class UECABridge;
+class FSocket;
+
+/**
+ * TCP Server Runnable
+ * 
+ * Handles incoming TCP connections from ECA clients.
+ * Runs on a separate thread and marshals commands to the game thread for execution.
+ */
+class ECABRIDGE_API FECAServerRunnable : public FRunnable
+{
+public:
+	FECAServerRunnable(UECABridge* InBridge, int32 InPort);
+	virtual ~FECAServerRunnable();
+
+	// FRunnable interface
+	virtual bool Init() override;
+	virtual uint32 Run() override;
+	virtual void Stop() override;
+
+private:
+	/** Handle a client connection */
+	void HandleClient(FSocket* ClientSocket);
+	
+	/** Read a complete message from the socket (handles length-prefixed protocol) */
+	bool ReadMessage(FSocket* ClientSocket, FString& OutMessage);
+	
+	/** Send a response to the client */
+	bool SendResponse(FSocket* ClientSocket, const FString& Response);
+
+	/** The bridge to route commands to */
+	UECABridge* Bridge;
+	
+	/** Listener socket */
+	FSocket* ListenerSocket;
+	
+	/** Port to listen on */
+	int32 Port;
+	
+	/** Flag to signal shutdown */
+	TAtomic<bool> bShouldStop;
+	
+	/** Buffer for receiving data */
+	static constexpr int32 BufferSize = 1024 * 1024; // 1MB buffer for large commands
+};
