Fix in ECABridge for static analysis finds (potential divide by zeros)

#Synced-CL 51504585

diff --git a/Engine/Plugins/Experimental/ECABridge/Source/ECABridge/Private/Commands/ECAMeshCommands.cpp b/Engine/Plugins/Experimental/ECABridge/Source/ECABridge/Private/Commands/ECAMeshCommands.cpp
--- a/Engine/Plugins/Experimental/ECABridge/Source/ECABridge/Private/Commands/ECAMeshCommands.cpp
+++ b/Engine/Plugins/Experimental/ECABridge/Source/ECABridge/Private/Commands/ECAMeshCommands.cpp
@@ -136,6 +136,9 @@ namespace MeshHelpers
 	// The default FCylinderGenerator has continuous UVs that wrap incorrectly at the seam
 	void GenerateCylinderWithUVSeam(FDynamicMesh3& OutMesh, double Radius, double Height, int32 RadialSegments, int32 HeightSegments, bool bCapped, bool bReverseOrientation)
 	{
+		RadialSegments = FMath::Max(RadialSegments, 3);
+		HeightSegments = FMath::Max(HeightSegments, 1);
+
 		OutMesh.Clear();
 		OutMesh.EnableAttributes();
 		OutMesh.Attributes()->SetNumUVLayers(1);
@@ -5555,6 +5558,7 @@ FECACommandResult FECACommand_MeshDisplaceFromDepth::Execute(const TSharedPtr<FJ
 
 	int32 Resolution = 512;
 	GetIntParam(Params, TEXT("resolution"), Resolution, false);
+	Resolution = FMath::Max(Resolution, 2);
 
 	// Capture current depth
 	UTextureRenderTarget2D* RenderTarget = NewObject<UTextureRenderTarget2D>();
