ECABridge: Fixes for unity builds

#Synced-CL 51504585

diff --git a/Engine/Plugins/Experimental/ECABridge/Source/ECABridge/Private/Commands/ECAAssetCommands.cpp b/Engine/Plugins/Experimental/ECABridge/Source/ECABridge/Private/Commands/ECAAssetCommands.cpp
--- a/Engine/Plugins/Experimental/ECABridge/Source/ECABridge/Private/Commands/ECAAssetCommands.cpp
+++ b/Engine/Plugins/Experimental/ECABridge/Source/ECABridge/Private/Commands/ECAAssetCommands.cpp
@@ -85,7 +85,7 @@ REGISTER_ECA_COMMAND(FECACommand_GetAssetThumbnails)
 //------------------------------------------------------------------------------
 
 // Collect compilation errors from material expressions
-static TArray<TSharedPtr<FJsonValue>> GetMaterialCompilationErrors(UMaterial* Material)
+static TArray<TSharedPtr<FJsonValue>> GetAssetMaterialCompilationErrors(UMaterial* Material)
 {
 	TArray<TSharedPtr<FJsonValue>> ErrorsArray;
 	
@@ -112,7 +112,7 @@ static TArray<TSharedPtr<FJsonValue>> GetMaterialCompilationErrors(UMaterial* Ma
 // Add compilation errors to result if any exist
 static void AddMaterialCompilationErrorsToResult(TSharedPtr<FJsonObject>& Result, UMaterial* Material)
 {
-	TArray<TSharedPtr<FJsonValue>> Errors = GetMaterialCompilationErrors(Material);
+	TArray<TSharedPtr<FJsonValue>> Errors = GetAssetMaterialCompilationErrors(Material);
 	if (Errors.Num() > 0)
 	{
 		Result->SetArrayField(TEXT("compilation_errors"), Errors);
@@ -354,7 +354,7 @@ static UTexture2D* LoadTextureByPath(const FString& TexturePath)
 	return LoadObject<UTexture2D>(nullptr, *TexturePath);
 }
 
-static UMaterial* LoadMaterialByPath(const FString& MaterialPath)
+static UMaterial* LoadMaterialAssetByPath(const FString& MaterialPath)
 {
 	return LoadObject<UMaterial>(nullptr, *MaterialPath);
 }
@@ -974,7 +974,7 @@ FECACommandResult FECACommand_GetMaterialInfo::Execute(const TSharedPtr<FJsonObj
 		return FECACommandResult::Error(TEXT("Missing required parameter: material_path"));
 	}
 	
-	UMaterial* Material = LoadMaterialByPath(MaterialPath);
+	UMaterial* Material = LoadMaterialAssetByPath(MaterialPath);
 	if (!Material)
 	{
 		return FECACommandResult::Error(FString::Printf(TEXT("Material not found: %s"), *MaterialPath));
@@ -1142,7 +1142,7 @@ FECACommandResult FECACommand_SetMaterialTextureParam::Execute(const TSharedPtr<
 		return FECACommandResult::Error(TEXT("Missing required parameter: texture_path"));
 	}
 	
-	UMaterial* Material = LoadMaterialByPath(MaterialPath);
+	UMaterial* Material = LoadMaterialAssetByPath(MaterialPath);
 	if (!Material)
 	{
 		return FECACommandResult::Error(FString::Printf(TEXT("Material not found: %s"), *MaterialPath));
diff --git a/Engine/Plugins/Experimental/ECABridge/Source/ECABridge/Private/Commands/ECABlueprintNodeCommands.cpp b/Engine/Plugins/Experimental/ECABridge/Source/ECABridge/Private/Commands/ECABlueprintNodeCommands.cpp
--- a/Engine/Plugins/Experimental/ECABridge/Source/ECABridge/Private/Commands/ECABlueprintNodeCommands.cpp
+++ b/Engine/Plugins/Experimental/ECABridge/Source/ECABridge/Private/Commands/ECABlueprintNodeCommands.cpp
@@ -210,7 +210,7 @@ static FVector2D GetNodePosition(const TSharedPtr<FJsonObject>& Params, UEdGraph
 
 // Helper to ensure a node has a valid GUID
 // Some nodes may not have their GUID properly initialized when created programmatically
-static void EnsureNodeHasValidGuid(UEdGraphNode* Node)
+static void EnsureBPNodeHasValidGuid(UEdGraphNode* Node)
 {
 	if (Node && !Node->NodeGuid.IsValid())
 	{
@@ -302,7 +302,7 @@ static TSharedPtr<FJsonObject> PinToJson(UEdGraphPin* Pin)
 static TSharedPtr<FJsonObject> NodeToJson(UEdGraphNode* Node)
 {
 	// Ensure the node has a valid GUID before serializing
-	EnsureNodeHasValidGuid(Node);
+	EnsureBPNodeHasValidGuid(Node);
 	
 	TSharedPtr<FJsonObject> NodeJson = MakeShared<FJsonObject>();
 	NodeJson->SetStringField(TEXT("node_id"), Node->NodeGuid.ToString());
@@ -514,7 +514,7 @@ FECACommandResult FECACommand_AddBlueprintEventNode::Execute(const TSharedPtr<FJ
 	EventNode->NodePosY = Position.Y;
 	EventGraph->AddNode(EventNode, false, false);
 	EventNode->AllocateDefaultPins();
-	EnsureNodeHasValidGuid(EventNode);
+	EnsureBPNodeHasValidGuid(EventNode);
 	EventNode->ReconstructNode();
 	
 	FBlueprintEditorUtils::MarkBlueprintAsStructurallyModified(Blueprint);
@@ -568,7 +568,7 @@ FECACommandResult FECACommand_AddBlueprintInputActionNode::Execute(const TShared
 	
 	EventGraph->AddNode(InputActionNode, false, false);
 	InputActionNode->AllocateDefaultPins();
-	EnsureNodeHasValidGuid(InputActionNode);
+	EnsureBPNodeHasValidGuid(InputActionNode);
 	
 	FBlueprintEditorUtils::MarkBlueprintAsStructurallyModified(Blueprint);
 	
@@ -660,7 +660,7 @@ FECACommandResult FECACommand_AddBlueprintFunctionNode::Execute(const TSharedPtr
 		CallNode->NodePosY = Position.Y;
 		Graph->AddNode(CallNode, false, false);
 		CallNode->AllocateDefaultPins();
-		EnsureNodeHasValidGuid(CallNode);
+		EnsureBPNodeHasValidGuid(CallNode);
 		CallNode->ReconstructNode();
 		
 		FBlueprintEditorUtils::MarkBlueprintAsStructurallyModified(Blueprint);
@@ -815,7 +815,7 @@ FECACommandResult FECACommand_AddBlueprintFunctionNode::Execute(const TSharedPtr
 	
 	Graph->AddNode(CallNode, false, false);
 	CallNode->AllocateDefaultPins();
-	EnsureNodeHasValidGuid(CallNode);
+	EnsureBPNodeHasValidGuid(CallNode);
 	
 	// Reconstruct node to ensure proper setup and error detection
 	CallNode->ReconstructNode();
@@ -1023,7 +1023,7 @@ FECACommandResult FECACommand_AddBlueprintSelfReference::Execute(const TSharedPt
 	
 	Graph->AddNode(SelfNode, false, false);
 	SelfNode->AllocateDefaultPins();
-	EnsureNodeHasValidGuid(SelfNode);
+	EnsureBPNodeHasValidGuid(SelfNode);
 	
 	FBlueprintEditorUtils::MarkBlueprintAsStructurallyModified(Blueprint);
 	
@@ -1077,7 +1077,7 @@ FECACommandResult FECACommand_AddBlueprintComponentReference::Execute(const TSha
 	
 	Graph->AddNode(GetNode, false, false);
 	GetNode->AllocateDefaultPins();
-	EnsureNodeHasValidGuid(GetNode);
+	EnsureBPNodeHasValidGuid(GetNode);
 	
 	FBlueprintEditorUtils::MarkBlueprintAsStructurallyModified(Blueprint);
 	
@@ -1195,7 +1195,7 @@ FECACommandResult FECACommand_AddBlueprintVariableGetNode::Execute(const TShared
 	
 	Graph->AddNode(GetNode, false, false);
 	GetNode->AllocateDefaultPins();
-	EnsureNodeHasValidGuid(GetNode);
+	EnsureBPNodeHasValidGuid(GetNode);
 	GetNode->ReconstructNode();
 	
 	FBlueprintEditorUtils::MarkBlueprintAsStructurallyModified(Blueprint);
@@ -1251,7 +1251,7 @@ FECACommandResult FECACommand_AddBlueprintVariableSetNode::Execute(const TShared
 	
 	Graph->AddNode(SetNode, false, false);
 	SetNode->AllocateDefaultPins();
-	EnsureNodeHasValidGuid(SetNode);
+	EnsureBPNodeHasValidGuid(SetNode);
 	SetNode->ReconstructNode();
 	
 	FBlueprintEditorUtils::MarkBlueprintAsStructurallyModified(Blueprint);
@@ -1624,7 +1624,7 @@ FECACommandResult FECACommand_BatchEditBlueprintNodes::Execute(const TSharedPtr<
 			CreatedNode->NodePosY = Position.Y;
 			Graph->AddNode(CreatedNode, false, false);
 			CreatedNode->AllocateDefaultPins();
-			EnsureNodeHasValidGuid(CreatedNode);
+			EnsureBPNodeHasValidGuid(CreatedNode);
 			
 			TempIdToNode.Add(TempId, CreatedNode);
 			
@@ -2863,7 +2863,7 @@ FECACommandResult FECACommand_AddComponentEventNode::Execute(const TSharedPtr<FJ
 	
 	EventGraph->AddNode(EventNode, false, false);
 	EventNode->AllocateDefaultPins();
-	EnsureNodeHasValidGuid(EventNode);
+	EnsureBPNodeHasValidGuid(EventNode);
 	
 	FBlueprintEditorUtils::MarkBlueprintAsStructurallyModified(Blueprint);
 	
@@ -2954,7 +2954,7 @@ FECACommandResult FECACommand_AddBlueprintMacroNode::Execute(const TSharedPtr<FJ
 	
 	Graph->AddNode(MacroNode, false, false);
 	MacroNode->AllocateDefaultPins();
-	EnsureNodeHasValidGuid(MacroNode);
+	EnsureBPNodeHasValidGuid(MacroNode);
 	MacroNode->ReconstructNode();
 	
 	FBlueprintEditorUtils::MarkBlueprintAsStructurallyModified(Blueprint);
@@ -3068,7 +3068,7 @@ FECACommandResult FECACommand_AddBlueprintCastNode::Execute(const TSharedPtr<FJs
 	
 	Graph->AddNode(CastNode, false, false);
 	CastNode->AllocateDefaultPins();
-	EnsureNodeHasValidGuid(CastNode);
+	EnsureBPNodeHasValidGuid(CastNode);
 	CastNode->ReconstructNode();
 	
 	FBlueprintEditorUtils::MarkBlueprintAsStructurallyModified(Blueprint);
@@ -3270,7 +3270,7 @@ FECACommandResult FECACommand_AddBlueprintFlowControlNode::Execute(const TShared
 		
 		Graph->AddNode(CreatedNode, false, false);
 		CreatedNode->AllocateDefaultPins();
-		EnsureNodeHasValidGuid(CreatedNode);
+		EnsureBPNodeHasValidGuid(CreatedNode);
 		CreatedNode->ReconstructNode();
 		
 		FBlueprintEditorUtils::MarkBlueprintAsStructurallyModified(Blueprint);
diff --git a/Engine/Plugins/Experimental/ECABridge/Source/ECABridge/Private/Commands/ECAMaterialNodeCommands.cpp b/Engine/Plugins/Experimental/ECABridge/Source/ECABridge/Private/Commands/ECAMaterialNodeCommands.cpp
--- a/Engine/Plugins/Experimental/ECABridge/Source/ECABridge/Private/Commands/ECAMaterialNodeCommands.cpp
+++ b/Engine/Plugins/Experimental/ECABridge/Source/ECABridge/Private/Commands/ECAMaterialNodeCommands.cpp
@@ -247,7 +247,7 @@ static FVector2D CalculateAutoMaterialNodePosition(UMaterial* Material)
 	return FVector2D(BestX, BestY);
 }
 
-static FVector2D GetNodePosition(const TSharedPtr<FJsonObject>& Params, UMaterial* Material = nullptr)
+static FVector2D GetMaterialNodePosition(const TSharedPtr<FJsonObject>& Params, UMaterial* Material = nullptr)
 {
 	FVector2D Position(0, 0);
 	const TSharedPtr<FJsonObject>* PosObj;
@@ -698,7 +698,7 @@ FECACommandResult FECACommand_AddMaterialNode::Execute(const TSharedPtr<FJsonObj
 		return FECACommandResult::Error(FString::Printf(TEXT("Unknown material expression type: %s"), *NodeType));
 	}
 	
-	FVector2D Position = GetNodePosition(Params, Material);
+	FVector2D Position = GetMaterialNodePosition(Params, Material);
 	
 	// Create the expression
 	Material->PreEditChange(nullptr);
@@ -2883,7 +2883,7 @@ FECACommandResult FECACommand_SetCustomNodeInputName::Execute(const TSharedPtr<F
 
 // Helper to convert JSON value to string for property setting
 // Handles JSON objects by converting them to UE4's ImportText format
-static FString JsonValueToPropertyString(const TSharedPtr<FJsonObject>& Params, const FString& Key)
+static FString MaterialJsonValueToPropertyString(const TSharedPtr<FJsonObject>& Params, const FString& Key)
 {
 	if (!Params->HasField(Key))
 	{
@@ -3008,7 +3008,7 @@ static FString JsonValueToPropertyString(const TSharedPtr<FJsonObject>& Params,
 }
 
 // Helper to convert FProperty value to JSON
-static TSharedPtr<FJsonValue> PropertyToJsonValue(FProperty* Property, const void* ValuePtr)
+static TSharedPtr<FJsonValue> MaterialPropertyToJsonValue(FProperty* Property, const void* ValuePtr)
 {
 	if (FBoolProperty* BoolProp = CastField<FBoolProperty>(Property))
 	{
@@ -3228,13 +3228,13 @@ FECACommandResult FECACommand_SetMaterialProperty::Execute(const TSharedPtr<FJso
 	}
 	
 	// Get the value as a string for ImportText
-	FString ValueStr = JsonValueToPropertyString(Params, TEXT("value"));
+	FString ValueStr = MaterialJsonValueToPropertyString(Params, TEXT("value"));
 	
 	// Get pointer to property value
 	void* ValuePtr = Property->ContainerPtrToValuePtr<void>(Material);
 	
 	// Store old value for reporting
-	TSharedPtr<FJsonValue> OldValue = PropertyToJsonValue(Property, ValuePtr);
+	TSharedPtr<FJsonValue> OldValue = MaterialPropertyToJsonValue(Property, ValuePtr);
 	
 	// Use ImportText to set the value (handles enums, bools, numbers, etc.)
 	const TCHAR* ImportResult = Property->ImportText_Direct(*ValueStr, ValuePtr, Material, PPF_None);
@@ -3244,7 +3244,7 @@ FECACommandResult FECACommand_SetMaterialProperty::Execute(const TSharedPtr<FJso
 	}
 	
 	// Get new value for reporting
-	TSharedPtr<FJsonValue> NewValue = PropertyToJsonValue(Property, ValuePtr);
+	TSharedPtr<FJsonValue> NewValue = MaterialPropertyToJsonValue(Property, ValuePtr);
 	
 	// Mark the material as needing recompilation
 	ClearMaterialCompilationErrors(Material);
@@ -3296,7 +3296,7 @@ FECACommandResult FECACommand_GetMaterialProperty::Execute(const TSharedPtr<FJso
 		}
 		
 		void* ValuePtr = Property->ContainerPtrToValuePtr<void>(Material);
-		Result->SetField(Property->GetName(), PropertyToJsonValue(Property, ValuePtr));
+		Result->SetField(Property->GetName(), MaterialPropertyToJsonValue(Property, ValuePtr));
 	}
 	else
 	{
@@ -3320,7 +3320,7 @@ FECACommandResult FECACommand_GetMaterialProperty::Execute(const TSharedPtr<FJso
 			}
 			
 			void* ValuePtr = Property->ContainerPtrToValuePtr<void>(Material);
-			PropertiesObj->SetField(Property->GetName(), PropertyToJsonValue(Property, ValuePtr));
+			PropertiesObj->SetField(Property->GetName(), MaterialPropertyToJsonValue(Property, ValuePtr));
 		}
 		
 		Result->SetObjectField(TEXT("properties"), PropertiesObj);
diff --git a/Engine/Plugins/Experimental/ECABridge/Source/ECABridge/Private/Commands/ECAMetasoundCommands.cpp b/Engine/Plugins/Experimental/ECABridge/Source/ECABridge/Private/Commands/ECAMetasoundCommands.cpp
--- a/Engine/Plugins/Experimental/ECABridge/Source/ECABridge/Private/Commands/ECAMetasoundCommands.cpp
+++ b/Engine/Plugins/Experimental/ECABridge/Source/ECABridge/Private/Commands/ECAMetasoundCommands.cpp
@@ -191,7 +191,7 @@ static bool ParseNodeId(const FString& NodeIdStr, FGuid& OutGuid)
 	return FGuid::Parse(NodeIdStr, OutGuid);
 }
 
-static TSharedPtr<FJsonObject> NodeToJson(const FMetasoundFrontendNode& Node, bool bIncludeDetails)
+static TSharedPtr<FJsonObject> MetasoundNodeToJson(const FMetasoundFrontendNode& Node, bool bIncludeDetails)
 {
 	TSharedPtr<FJsonObject> NodeObj = MakeShared<FJsonObject>();
 	NodeObj->SetStringField(TEXT("id"), GetNodeIdString(Node.GetID()));
@@ -404,7 +404,7 @@ FECACommandResult FECACommand_GetMetasoundNodes::Execute(const TSharedPtr<FJsonO
 	DocBuilder.IterateNodesByPredicate(
 		[&](const FMetasoundFrontendClass& NodeClass, const FMetasoundFrontendNode& Node, const FGuid& PageID)
 		{
-			NodesArray.Add(MakeShared<FJsonValueObject>(NodeToJson(Node, bIncludeDetails)));
+			NodesArray.Add(MakeShared<FJsonValueObject>(MetasoundNodeToJson(Node, bIncludeDetails)));
 		},
 		[](const FMetasoundFrontendNode& Node) { return true; }, // Include all nodes
 		nullptr, // No specific page ID
