ECABridge: fixed documentation for (rot) and exposed some more stuff to batch blueprint node gen

#Synced-CL 51504585

diff --git a/Engine/Plugins/Experimental/ECABridge/Source/ECABridge/Private/Commands/ECABlueprintLispCommands.cpp b/Engine/Plugins/Experimental/ECABridge/Source/ECABridge/Private/Commands/ECABlueprintLispCommands.cpp
--- a/Engine/Plugins/Experimental/ECABridge/Source/ECABridge/Private/Commands/ECABlueprintLispCommands.cpp
+++ b/Engine/Plugins/Experimental/ECABridge/Source/ECABridge/Private/Commands/ECABlueprintLispCommands.cpp
@@ -6342,7 +6342,7 @@ Literals:
   "hello"      - Strings
   true, false  - Booleans
   (vec x y z)  - Vector3
-  (rot p y r)  - Rotator
+  (rot r p y)  - Rotator (Roll Pitch Yaw)
   self         - Self reference
   nil          - Null value
 
@@ -6483,7 +6483,7 @@ Boolean:
 Vector Construction:
   (vec x y z)        - Make Vector (FVector)
   (vec2 x y)         - Make Vector2D
-  (rot pitch yaw roll) - Make Rotator
+  (rot roll pitch yaw) - Make Rotator
 
 Math Functions (use as regular function calls):
   (abs x)            - Absolute value
diff --git a/Engine/Plugins/Experimental/ECABridge/Source/ECABridge/Private/Commands/ECABlueprintNodeCommands.cpp b/Engine/Plugins/Experimental/ECABridge/Source/ECABridge/Private/Commands/ECABlueprintNodeCommands.cpp
--- a/Engine/Plugins/Experimental/ECABridge/Source/ECABridge/Private/Commands/ECABlueprintNodeCommands.cpp
+++ b/Engine/Plugins/Experimental/ECABridge/Source/ECABridge/Private/Commands/ECABlueprintNodeCommands.cpp
@@ -1421,31 +1421,112 @@ FECACommandResult FECACommand_BatchEditBlueprintNodes::Execute(const TSharedPtr<
 				}
 			}
 			
-			ClassesToSearch.Add(UKismetSystemLibrary::StaticClass());
-			ClassesToSearch.Add(UGameplayStatics::StaticClass());
-			ClassesToSearch.Add(AActor::StaticClass());
-			ClassesToSearch.Add(APawn::StaticClass());
-			ClassesToSearch.Add(ACharacter::StaticClass());
-			ClassesToSearch.Add(Blueprint->ParentClass);
-			
-			for (UClass* Class : ClassesToSearch)
+			// First, check if this is a function defined in the Blueprint itself
+			bool bIsBlueprintFunction = false;
+			for (UEdGraph* FuncGraph : Blueprint->FunctionGraphs)
 			{
-				if (Class)
+				if (FuncGraph && FuncGraph->GetFName() == FName(*FunctionName))
 				{
-					Function = Class->FindFunctionByName(*FunctionName);
-					if (Function) break;
+					bIsBlueprintFunction = true;
+					break;
 				}
 			}
 			
-			if (!Function)
+			if (bIsBlueprintFunction)
 			{
-				Errors.Add(FString::Printf(TEXT("Node '%s': Function not found: %s"), *TempId, *FunctionName));
-				continue;
+				UK2Node_CallFunction* CallNode = NewObject<UK2Node_CallFunction>(Graph);
+				CallNode->FunctionReference.SetSelfMember(FName(*FunctionName));
+				CreatedNode = CallNode;
+			}
+			else
+			{
+				// Add common classes for engine/library functions
+				ClassesToSearch.Add(UKismetSystemLibrary::StaticClass());
+				ClassesToSearch.Add(UKismetMathLibrary::StaticClass());
+				ClassesToSearch.Add(UKismetStringLibrary::StaticClass());
+				ClassesToSearch.Add(UKismetArrayLibrary::StaticClass());
+				ClassesToSearch.Add(UGameplayStatics::StaticClass());
+				ClassesToSearch.Add(AActor::StaticClass());
+				ClassesToSearch.Add(APawn::StaticClass());
+				ClassesToSearch.Add(ACharacter::StaticClass());
+				ClassesToSearch.Add(USceneComponent::StaticClass());
+				ClassesToSearch.Add(UPrimitiveComponent::StaticClass());
+				ClassesToSearch.Add(Blueprint->ParentClass);
+				
+				// Add Geometry Script library classes
+				static const TCHAR* GeometryScriptClassPaths[] = {
+					TEXT("/Script/GeometryScriptingCore.GeometryScriptLibrary_MeshPrimitiveFunctions"),
+					TEXT("/Script/GeometryScriptingCore.GeometryScriptLibrary_MeshBasicEditFunctions"),
+					TEXT("/Script/GeometryScriptingCore.GeometryScriptLibrary_MeshTransformFunctions"),
+					TEXT("/Script/GeometryScriptingCore.GeometryScriptLibrary_MeshBooleanFunctions"),
+					TEXT("/Script/GeometryScriptingCore.GeometryScriptLibrary_MeshModelingFunctions"),
+					TEXT("/Script/GeometryScriptingCore.GeometryScriptLibrary_MeshNormalsFunctions"),
+					TEXT("/Script/GeometryScriptingCore.GeometryScriptLibrary_MeshUVFunctions"),
+					TEXT("/Script/GeometryScriptingCore.GeometryScriptLibrary_MeshMaterialFunctions"),
+					TEXT("/Script/GeometryScriptingCore.GeometryScriptLibrary_MeshDeformFunctions"),
+					TEXT("/Script/GeometryScriptingCore.GeometryScriptLibrary_MeshSubdivideFunctions"),
+					TEXT("/Script/GeometryScriptingCore.GeometryScriptLibrary_MeshSimplifyFunctions"),
+					TEXT("/Script/GeometryScriptingCore.GeometryScriptLibrary_MeshRepairFunctions"),
+					TEXT("/Script/GeometryScriptingCore.GeometryScriptLibrary_MeshQueryFunctions"),
+					TEXT("/Script/GeometryScriptingCore.GeometryScriptLibrary_MeshSpatialFunctions"),
+					TEXT("/Script/GeometryScriptingCore.GeometryScriptLibrary_MeshSelectionFunctions"),
+					TEXT("/Script/GeometryScriptingCore.GeometryScriptLibrary_ContainmentFunctions"),
+					TEXT("/Script/GeometryScriptingCore.GeometryScriptLibrary_CollisionFunctions"),
+					TEXT("/Script/GeometryScriptingCore.GeometryScriptLibrary_MeshVoxelFunctions"),
+					TEXT("/Script/GeometryScriptingCore.GeometryScriptLibrary_PolyPathFunctions"),
+					TEXT("/Script/GeometryScriptingCore.GeometryScriptLibrary_ListUtilityFunctions"),
+					TEXT("/Script/GeometryScriptingEditor.GeometryScriptLibrary_CreateNewAssetUtilityFunctions"),
+				};
+				for (const TCHAR* ClassPath : GeometryScriptClassPaths)
+				{
+					UClass* GeoClass = FindObject<UClass>(nullptr, ClassPath);
+					if (!GeoClass) { GeoClass = LoadObject<UClass>(nullptr, ClassPath); }
+					if (!GeoClass) { GeoClass = StaticLoadClass(UObject::StaticClass(), nullptr, ClassPath); }
+					if (GeoClass) { ClassesToSearch.Add(GeoClass); }
+				}
+				if (Blueprint->GeneratedClass)
+				{
+					ClassesToSearch.Add(Blueprint->GeneratedClass);
+				}
+				
+				for (UClass* Class : ClassesToSearch)
+				{
+					if (Class)
+					{
+						Function = Class->FindFunctionByName(*FunctionName);
+						if (Function) break;
+						// Also try with K2_ prefix
+						Function = Class->FindFunctionByName(*FString::Printf(TEXT("K2_%s"), *FunctionName));
+						if (Function) break;
+					}
+				}
+				
+				// If still not found, do a broader search across all BlueprintFunctionLibrary classes
+				if (!Function)
+				{
+					for (TObjectIterator<UClass> It; It; ++It)
+					{
+						UClass* TestClass = *It;
+						if (TestClass->IsChildOf(UBlueprintFunctionLibrary::StaticClass()) && !TestClass->HasAnyClassFlags(CLASS_Abstract))
+						{
+							Function = TestClass->FindFunctionByName(*FunctionName);
+							if (Function) break;
+							Function = TestClass->FindFunctionByName(*FString::Printf(TEXT("K2_%s"), *FunctionName));
+							if (Function) break;
+						}
+					}
+				}
+				
+				if (!Function)
+				{
+					Errors.Add(FString::Printf(TEXT("Node '%s': Function not found: %s. Searched %d classes including math, string, array, and Geometry Script libraries."), *TempId, *FunctionName, ClassesToSearch.Num()));
+					continue;
+				}
+				
+				UK2Node_CallFunction* CallNode = NewObject<UK2Node_CallFunction>(Graph);
+				CallNode->SetFromFunction(Function);
+				CreatedNode = CallNode;
 			}
-			
-			UK2Node_CallFunction* CallNode = NewObject<UK2Node_CallFunction>(Graph);
-			CallNode->SetFromFunction(Function);
-			CreatedNode = CallNode;
 		}
 		else if (NodeType == TEXT("variable_get"))
 		{
