ECABridge: more fixes for unity builds

#Synced-CL 51504585

diff --git a/Engine/Plugins/Experimental/ECABridge/Source/ECABridge/Private/Commands/ECAEventCommands.cpp b/Engine/Plugins/Experimental/ECABridge/Source/ECABridge/Private/Commands/ECAEventCommands.cpp
--- a/Engine/Plugins/Experimental/ECABridge/Source/ECABridge/Private/Commands/ECAEventCommands.cpp
+++ b/Engine/Plugins/Experimental/ECABridge/Source/ECABridge/Private/Commands/ECAEventCommands.cpp
@@ -246,7 +246,7 @@ REGISTER_ECA_COMMAND(FECACommand_GetEventQueueStatus)
 
 namespace ECAEvents
 {
-	void QueueSelectionChanged(const FString& ActorName, const FString& ActorPath)
+	ECABRIDGE_API void QueueSelectionChanged(const FString& ActorName, const FString& ActorPath)
 	{
 		TSharedPtr<FJsonObject> Data = MakeShared<FJsonObject>();
 		Data->SetStringField(TEXT("actor_name"), ActorName);
@@ -256,7 +256,7 @@ namespace ECAEvents
 		FECAEventQueue::Get().EnqueueEvent(TEXT("selection_changed"), Message, Data, true);
 	}
 
-	void QueueAssetImported(const FString& AssetPath, const FString& AssetType)
+	ECABRIDGE_API void QueueAssetImported(const FString& AssetPath, const FString& AssetType)
 	{
 		TSharedPtr<FJsonObject> Data = MakeShared<FJsonObject>();
 		Data->SetStringField(TEXT("asset_path"), AssetPath);
@@ -266,7 +266,7 @@ namespace ECAEvents
 		FECAEventQueue::Get().EnqueueEvent(TEXT("asset_imported"), Message, Data, false);
 	}
 
-	void QueueAssetSaved(const FString& AssetPath)
+	ECABRIDGE_API void QueueAssetSaved(const FString& AssetPath)
 	{
 		TSharedPtr<FJsonObject> Data = MakeShared<FJsonObject>();
 		Data->SetStringField(TEXT("asset_path"), AssetPath);
@@ -275,7 +275,7 @@ namespace ECAEvents
 		FECAEventQueue::Get().EnqueueEvent(TEXT("asset_saved"), Message, Data, false);
 	}
 
-	void QueueBlueprintCompiled(const FString& BlueprintPath, bool bSuccess, const FString& ErrorMessage)
+	ECABRIDGE_API void QueueBlueprintCompiled(const FString& BlueprintPath, bool bSuccess, const FString& ErrorMessage)
 	{
 		TSharedPtr<FJsonObject> Data = MakeShared<FJsonObject>();
 		Data->SetStringField(TEXT("blueprint_path"), BlueprintPath);
@@ -292,7 +292,7 @@ namespace ECAEvents
 		FECAEventQueue::Get().EnqueueEvent(TEXT("blueprint_compiled"), Message, Data, !bSuccess);  // Only trigger AI on failures
 	}
 
-	void QueueCustomEvent(const FString& EventType, const FString& Message, TSharedPtr<FJsonObject> Data, bool bTriggerAIResponse)
+	ECABRIDGE_API void QueueCustomEvent(const FString& EventType, const FString& Message, TSharedPtr<FJsonObject> Data, bool bTriggerAIResponse)
 	{
 		FECAEventQueue::Get().EnqueueEvent(EventType, Message, Data, bTriggerAIResponse);
 	}
