ECaBridge: fixed PVS warnings

#Synced-CL 51504585

diff --git a/Engine/Plugins/Experimental/ECABridge/Source/ECABridge/Private/Commands/ECAMeshCommands.cpp b/Engine/Plugins/Experimental/ECABridge/Source/ECABridge/Private/Commands/ECAMeshCommands.cpp
--- a/Engine/Plugins/Experimental/ECABridge/Source/ECABridge/Private/Commands/ECAMeshCommands.cpp
+++ b/Engine/Plugins/Experimental/ECABridge/Source/ECABridge/Private/Commands/ECAMeshCommands.cpp
@@ -159,14 +159,14 @@ namespace MeshHelpers
 
 			for (int32 Seg = 0; Seg <= RadialSegments; Seg++)
 			{
-				double Angle = 2.0 * PI * Seg / RadialSegments;
+				double Angle = 2.0 * PI * Seg / RadialSegments; // -V609
 				double X = Radius * FMath::Cos(Angle);
 				double Y = Radius * FMath::Sin(Angle);
 
 				OutMesh.AppendVertex(FVector3d(X, Y, Z));
 
 				// UV: U goes from 0 to 1 around the cylinder, V goes from 0 to 1 along height
-				double U = (double)Seg / RadialSegments;
+				double U = (double)Seg / RadialSegments; // -V609
 				UVOverlay->AppendElement(FVector2f((float)U, (float)V));
 			}
 		}
@@ -222,7 +222,7 @@ namespace MeshHelpers
 			TArray<int32> BottomEdgeVerts, BottomEdgeUVs;
 			for (int32 Seg = 0; Seg <= RadialSegments; Seg++)
 			{
-				double Angle = 2.0 * PI * Seg / RadialSegments;
+				double Angle = 2.0 * PI * Seg / RadialSegments; // -V609
 				double X = Radius * FMath::Cos(Angle);
 				double Y = Radius * FMath::Sin(Angle);
 
@@ -261,7 +261,7 @@ namespace MeshHelpers
 			TArray<int32> TopEdgeVerts, TopEdgeUVs;
 			for (int32 Seg = 0; Seg <= RadialSegments; Seg++)
 			{
-				double Angle = 2.0 * PI * Seg / RadialSegments;
+				double Angle = 2.0 * PI * Seg / RadialSegments; // -V609
 				double X = Radius * FMath::Cos(Angle);
 				double Y = Radius * FMath::Sin(Angle);
 
@@ -5675,8 +5675,8 @@ FECACommandResult FECACommand_MeshDisplaceFromDepth::Execute(const TSharedPtr<FJ
 				for (int32 x = 0; x < Resolution; x++)
 				{
 					// Bilinear sample from source
-					float SrcX = (float)x / (Resolution - 1) * (TargetWidth - 1);
-					float SrcY = (float)y / (Resolution - 1) * (TargetHeight - 1);
+					float SrcX = (float)x / (Resolution - 1) * (TargetWidth - 1); // -V609
+					float SrcY = (float)y / (Resolution - 1) * (TargetHeight - 1); // -V609
 					
 					int32 X0 = FMath::Clamp((int32)SrcX, 0, TargetWidth - 1);
 					int32 Y0 = FMath::Clamp((int32)SrcY, 0, TargetHeight - 1);
