[ECABridge] Add save_asset, save_all_dirty, save_and_compile_blueprint, is_asset_dirty MCP tools

Currently there is no way to save assets to disk via MCP. Every edit session
requires manual Ctrl+S to persist changes. This adds four new tools:

- save_asset: Save a single asset by path (UPackage::SavePackage)
- save_all_dirty: Save all modified packages at once
- save_and_compile_blueprint: Compile + save in one call
- is_asset_dirty: Query dirty state for verification

Self-registered via REGISTER_ECA_COMMAND -- no existing files modified.

#rnx
#Synced-CL 51706747

#Synced-CL 51706747

diff --git a/Engine/Plugins/Experimental/ECABridge/Source/ECABridge/Private/Commands/ECASaveCommands.cpp b/Engine/Plugins/Experimental/ECABridge/Source/ECABridge/Private/Commands/ECASaveCommands.cpp
new file mode 100644
--- /dev/null
+++ b/Engine/Plugins/Experimental/ECABridge/Source/ECABridge/Private/Commands/ECASaveCommands.cpp
@@ -0,0 +1,234 @@
+// Copyright Epic Games, Inc. All Rights Reserved.
+
+#include "Commands/ECASaveCommands.h"
+#include "Engine/Blueprint.h"
+#include "Kismet2/KismetEditorUtilities.h"
+#include "Misc/PackageName.h"
+#include "UObject/Package.h"
+#include "UObject/SavePackage.h"
+#include "UObject/UObjectIterator.h"
+
+// Self-registration — no other files need modification
+REGISTER_ECA_COMMAND(FECACommand_IsAssetDirty)
+REGISTER_ECA_COMMAND(FECACommand_SaveAsset)
+REGISTER_ECA_COMMAND(FECACommand_SaveAllDirty)
+REGISTER_ECA_COMMAND(FECACommand_SaveAndCompileBlueprint)
+
+//----------------------------------------------------------------------
+// Shared helper: resolve asset path to loaded UObject + its UPackage
+//----------------------------------------------------------------------
+namespace ECASaveHelpers
+{
+	struct FResolvedAsset
+	{
+		UObject* Asset = nullptr;
+		UPackage* Package = nullptr;
+	};
+
+	static FECACommandResult ResolveAsset(const FString& AssetPath, FResolvedAsset& OutResolved)
+	{
+		OutResolved.Asset = LoadObject<UObject>(nullptr, *AssetPath);
+		if (!OutResolved.Asset)
+		{
+			return FECACommandResult::Error(FString::Printf(TEXT("Asset not found: %s"), *AssetPath));
+		}
+
+		OutResolved.Package = OutResolved.Asset->GetPackage();
+		if (!OutResolved.Package)
+		{
+			return FECACommandResult::Error(FString::Printf(TEXT("Could not get package for asset: %s"), *AssetPath));
+		}
+
+		return FECACommandResult::Success();
+	}
+
+	static FString GetPackageExtension(UPackage* Package)
+	{
+		return Package->ContainsMap()
+			? FPackageName::GetMapPackageExtension()
+			: FPackageName::GetAssetPackageExtension();
+	}
+
+	static bool SavePackageToDisk(UPackage* Package, UObject* Asset)
+	{
+		FString PackageFileName = FPackageName::LongPackageNameToFilename(
+			Package->GetName(), GetPackageExtension(Package));
+
+		FSavePackageArgs SaveArgs;
+		SaveArgs.TopLevelFlags = RF_Standalone;
+		bool bSaved = UPackage::SavePackage(Package, Asset, *PackageFileName, SaveArgs);
+
+		if (bSaved)
+		{
+			UE_LOG(LogTemp, Log, TEXT("[ECABridge] Saved %s"), *PackageFileName);
+		}
+		else
+		{
+			UE_LOG(LogTemp, Warning, TEXT("[ECABridge] Failed to save %s"), *PackageFileName);
+		}
+
+		return bSaved;
+	}
+
+	static bool ShouldSkipPackage(UPackage* Package)
+	{
+		if (Package->HasAnyFlags(RF_Transient))
+		{
+			return true;
+		}
+
+		FString PackageName = Package->GetName();
+		return PackageName == TEXT("Transient")
+			|| PackageName.StartsWith(TEXT("/Temp/"))
+			|| PackageName.StartsWith(TEXT("/Engine/"))
+			|| PackageName.StartsWith(TEXT("/Script/"));
+	}
+}
+
+//----------------------------------------------------------------------
+// is_asset_dirty
+//----------------------------------------------------------------------
+
+FECACommandResult FECACommand_IsAssetDirty::Execute(const TSharedPtr<FJsonObject>& Params)
+{
+	FString AssetPath;
+	if (!GetStringParam(Params, TEXT("asset_path"), AssetPath))
+	{
+		return FECACommandResult::Error(TEXT("Missing required parameter: asset_path"));
+	}
+
+	ECASaveHelpers::FResolvedAsset Resolved;
+	FECACommandResult ResolveResult = ECASaveHelpers::ResolveAsset(AssetPath, Resolved);
+	if (!ResolveResult.bSuccess)
+	{
+		return ResolveResult;
+	}
+
+	TSharedPtr<FJsonObject> Result = MakeResult();
+	Result->SetStringField(TEXT("asset_path"), AssetPath);
+	Result->SetBoolField(TEXT("is_dirty"), Resolved.Package->IsDirty());
+	Result->SetStringField(TEXT("package_name"), Resolved.Package->GetName());
+	return FECACommandResult::Success(Result);
+}
+
+//----------------------------------------------------------------------
+// save_asset
+//----------------------------------------------------------------------
+
+FECACommandResult FECACommand_SaveAsset::Execute(const TSharedPtr<FJsonObject>& Params)
+{
+	FString AssetPath;
+	if (!GetStringParam(Params, TEXT("asset_path"), AssetPath))
+	{
+		return FECACommandResult::Error(TEXT("Missing required parameter: asset_path"));
+	}
+
+	ECASaveHelpers::FResolvedAsset Resolved;
+	FECACommandResult ResolveResult = ECASaveHelpers::ResolveAsset(AssetPath, Resolved);
+	if (!ResolveResult.bSuccess)
+	{
+		return ResolveResult;
+	}
+
+	bool bWasDirty = Resolved.Package->IsDirty();
+
+	if (!ECASaveHelpers::SavePackageToDisk(Resolved.Package, Resolved.Asset))
+	{
+		return FECACommandResult::Error(FString::Printf(
+			TEXT("Failed to save package: %s"), *Resolved.Package->GetName()));
+	}
+
+	TSharedPtr<FJsonObject> Result = MakeResult();
+	Result->SetBoolField(TEXT("saved"), true);
+	Result->SetStringField(TEXT("asset_path"), AssetPath);
+	Result->SetStringField(TEXT("package_name"), Resolved.Package->GetName());
+	Result->SetBoolField(TEXT("was_dirty"), bWasDirty);
+	return FECACommandResult::Success(Result);
+}
+
+//----------------------------------------------------------------------
+// save_all_dirty
+//----------------------------------------------------------------------
+
+FECACommandResult FECACommand_SaveAllDirty::Execute(const TSharedPtr<FJsonObject>& Params)
+{
+	TArray<TSharedPtr<FJsonValue>> SavedArray;
+	TArray<TSharedPtr<FJsonValue>> FailedArray;
+	int32 SavedCount = 0;
+	int32 FailedCount = 0;
+
+	for (TObjectIterator<UPackage> It; It; ++It)
+	{
+		UPackage* Package = *It;
+
+		if (!Package->IsDirty() || ECASaveHelpers::ShouldSkipPackage(Package))
+		{
+			continue;
+		}
+
+		FString PackageName = Package->GetName();
+
+		if (ECASaveHelpers::SavePackageToDisk(Package, nullptr))
+		{
+			SavedCount++;
+			SavedArray.Add(MakeShared<FJsonValueString>(PackageName));
+		}
+		else
+		{
+			FailedCount++;
+			FailedArray.Add(MakeShared<FJsonValueString>(PackageName));
+		}
+	}
+
+	TSharedPtr<FJsonObject> Result = MakeResult();
+	Result->SetNumberField(TEXT("saved_count"), SavedCount);
+	Result->SetNumberField(TEXT("failed_count"), FailedCount);
+	Result->SetArrayField(TEXT("saved_packages"), SavedArray);
+	if (FailedCount > 0)
+	{
+		Result->SetArrayField(TEXT("failed_packages"), FailedArray);
+	}
+	return FECACommandResult::Success(Result);
+}
+
+//----------------------------------------------------------------------
+// save_and_compile_blueprint
+//----------------------------------------------------------------------
+
+FECACommandResult FECACommand_SaveAndCompileBlueprint::Execute(const TSharedPtr<FJsonObject>& Params)
+{
+	FString BlueprintPath;
+	if (!GetStringParam(Params, TEXT("blueprint_path"), BlueprintPath))
+	{
+		return FECACommandResult::Error(TEXT("Missing required parameter: blueprint_path"));
+	}
+
+	UBlueprint* Blueprint = LoadBlueprintByPath(BlueprintPath);
+	if (!Blueprint)
+	{
+		return FECACommandResult::Error(FString::Printf(
+			TEXT("Blueprint not found: %s"), *BlueprintPath));
+	}
+
+	// Step 1: Compile
+	FKismetEditorUtilities::CompileBlueprint(Blueprint);
+	bool bHasErrors = Blueprint->Status == BS_Error;
+
+	// Step 2: Save (always — even with compile errors, persist WIP)
+	UPackage* Package = Blueprint->GetPackage();
+	bool bSaved = ECASaveHelpers::SavePackageToDisk(Package, Blueprint);
+
+	TSharedPtr<FJsonObject> Result = MakeResult();
+	Result->SetBoolField(TEXT("compiled"), true);
+	Result->SetBoolField(TEXT("has_errors"), bHasErrors);
+	Result->SetStringField(TEXT("status"), bHasErrors ? TEXT("Error") : TEXT("Success"));
+	Result->SetStringField(TEXT("blueprint_path"), BlueprintPath);
+	Result->SetStringField(TEXT("package_name"), Package->GetName());
+	Result->SetBoolField(TEXT("saved"), bSaved);
+	if (!bSaved)
+	{
+		Result->SetStringField(TEXT("save_error"), FString::Printf(
+			TEXT("Failed to save package: %s"), *Package->GetName()));
+	}
+	return FECACommandResult::Success(Result);
+}
diff --git a/Engine/Plugins/Experimental/ECABridge/Source/ECABridge/Public/Commands/ECASaveCommands.h b/Engine/Plugins/Experimental/ECABridge/Source/ECABridge/Public/Commands/ECASaveCommands.h
new file mode 100644
--- /dev/null
+++ b/Engine/Plugins/Experimental/ECABridge/Source/ECABridge/Public/Commands/ECASaveCommands.h
@@ -0,0 +1,71 @@
+// Copyright Epic Games, Inc. All Rights Reserved.
+
+#pragma once
+
+#include "Commands/ECACommand.h"
+
+/**
+ * Check if an asset has unsaved modifications
+ */
+class FECACommand_IsAssetDirty : public IECACommand
+{
+public:
+	virtual FString GetName() const override { return TEXT("is_asset_dirty"); }
+	virtual FString GetDescription() const override { return TEXT("Check if an asset has unsaved modifications"); }
+	virtual FString GetCategory() const override { return TEXT("Asset"); }
+	virtual TArray<FECACommandParam> GetParameters() const override
+	{
+		return {
+			{ TEXT("asset_path"), TEXT("string"), TEXT("Path to the asset (e.g. /Game/UI/MyWidget)"), true }
+		};
+	}
+	virtual FECACommandResult Execute(const TSharedPtr<FJsonObject>& Params) override;
+};
+
+/**
+ * Save a single asset to disk
+ */
+class FECACommand_SaveAsset : public IECACommand
+{
+public:
+	virtual FString GetName() const override { return TEXT("save_asset"); }
+	virtual FString GetDescription() const override { return TEXT("Save a single asset to disk by path"); }
+	virtual FString GetCategory() const override { return TEXT("Asset"); }
+	virtual TArray<FECACommandParam> GetParameters() const override
+	{
+		return {
+			{ TEXT("asset_path"), TEXT("string"), TEXT("Path to the asset to save (e.g. /Game/UI/MyWidget)"), true }
+		};
+	}
+	virtual FECACommandResult Execute(const TSharedPtr<FJsonObject>& Params) override;
+};
+
+/**
+ * Save all dirty (modified) packages to disk
+ */
+class FECACommand_SaveAllDirty : public IECACommand
+{
+public:
+	virtual FString GetName() const override { return TEXT("save_all_dirty"); }
+	virtual FString GetDescription() const override { return TEXT("Save all modified assets to disk"); }
+	virtual FString GetCategory() const override { return TEXT("Asset"); }
+	virtual FECACommandResult Execute(const TSharedPtr<FJsonObject>& Params) override;
+};
+
+/**
+ * Compile a Blueprint and save it to disk in one operation
+ */
+class FECACommand_SaveAndCompileBlueprint : public IECACommand
+{
+public:
+	virtual FString GetName() const override { return TEXT("save_and_compile_blueprint"); }
+	virtual FString GetDescription() const override { return TEXT("Compile a Blueprint and save it to disk in one operation"); }
+	virtual FString GetCategory() const override { return TEXT("Blueprint"); }
+	virtual TArray<FECACommandParam> GetParameters() const override
+	{
+		return {
+			{ TEXT("blueprint_path"), TEXT("string"), TEXT("Path to the Blueprint to compile and save"), true }
+		};
+	}
+	virtual FECACommandResult Execute(const TSharedPtr<FJsonObject>& Params) override;
+};
