ECABridge: Use engine WriteRowAsJSON for DataTable commands

Replace manual property serialization with UDataTable::WriteRowAsJSON().
Fixes FInstancedStruct fields which serialized
as empty {}. Schema uses default-initialized row via same path to guarantee
format consistency with data output. ~250 lines of custom serializers removed.

#tests Tested commands, validated before/after output
#Synced-CL 51713787

diff --git a/Engine/Plugins/Experimental/ECABridge/Source/ECABridge/Private/Commands/ECADataTableCommands.cpp b/Engine/Plugins/Experimental/ECABridge/Source/ECABridge/Private/Commands/ECADataTableCommands.cpp
--- a/Engine/Plugins/Experimental/ECABridge/Source/ECABridge/Private/Commands/ECADataTableCommands.cpp
+++ b/Engine/Plugins/Experimental/ECABridge/Source/ECABridge/Private/Commands/ECADataTableCommands.cpp
@@ -17,233 +17,6 @@ REGISTER_ECA_COMMAND(FECACommand_GetDataTableSchema)
 REGISTER_ECA_COMMAND(FECACommand_GetDataTableRows)
 REGISTER_ECA_COMMAND(FECACommand_SetDataTableRow)
 REGISTER_ECA_COMMAND(FECACommand_DeleteDataTableRow)
-
-//------------------------------------------------------------------------------
-// Helper: Convert a single property value to a JSON value
-//------------------------------------------------------------------------------
-static TSharedPtr<FJsonValue> DataTablePropertyToJsonValue(FProperty* Property, const void* ValuePtr)
-{
-	if (!Property || !ValuePtr)
-	{
-		return MakeShared<FJsonValueNull>();
-	}
-
-	if (FBoolProperty* BoolProp = CastField<FBoolProperty>(Property))
-	{
-		return MakeShared<FJsonValueBoolean>(BoolProp->GetPropertyValue(ValuePtr));
-	}
-	else if (FNumericProperty* NumProp = CastField<FNumericProperty>(Property))
-	{
-		if (NumProp->IsFloatingPoint())
-		{
-			return MakeShared<FJsonValueNumber>(NumProp->GetFloatingPointPropertyValue(ValuePtr));
-		}
-		else if (NumProp->IsInteger())
-		{
-			return MakeShared<FJsonValueNumber>(static_cast<double>(NumProp->GetSignedIntPropertyValue(ValuePtr)));
-		}
-	}
-	else if (FStrProperty* StrProp = CastField<FStrProperty>(Property))
-	{
-		return MakeShared<FJsonValueString>(StrProp->GetPropertyValue(ValuePtr));
-	}
-	else if (FNameProperty* NameProp = CastField<FNameProperty>(Property))
-	{
-		return MakeShared<FJsonValueString>(NameProp->GetPropertyValue(ValuePtr).ToString());
-	}
-	else if (FTextProperty* TextProp = CastField<FTextProperty>(Property))
-	{
-		return MakeShared<FJsonValueString>(TextProp->GetPropertyValue(ValuePtr).ToString());
-	}
-	else if (FEnumProperty* EnumProp = CastField<FEnumProperty>(Property))
-	{
-		UEnum* Enum = EnumProp->GetEnum();
-		FNumericProperty* UnderlyingProp = EnumProp->GetUnderlyingProperty();
-		int64 EnumValue = UnderlyingProp->GetSignedIntPropertyValue(ValuePtr);
-		FString EnumName = Enum->GetNameStringByValue(EnumValue);
-		return MakeShared<FJsonValueString>(EnumName);
-	}
-	else if (FByteProperty* ByteProp = CastField<FByteProperty>(Property))
-	{
-		if (UEnum* Enum = ByteProp->GetIntPropertyEnum())
-		{
-			uint8 ByteValue = *static_cast<const uint8*>(ValuePtr);
-			FString EnumName = Enum->GetNameStringByValue(ByteValue);
-			return MakeShared<FJsonValueString>(EnumName);
-		}
-		else
-		{
-			return MakeShared<FJsonValueNumber>(static_cast<double>(*static_cast<const uint8*>(ValuePtr)));
-		}
-	}
-	else if (FStructProperty* StructProp = CastField<FStructProperty>(Property))
-	{
-		// For structs, recurse into their properties
-		TSharedPtr<FJsonObject> StructObj = MakeShared<FJsonObject>();
-		for (TFieldIterator<FProperty> PropIt(StructProp->Struct); PropIt; ++PropIt)
-		{
-			FProperty* InnerProp = *PropIt;
-			const void* InnerValuePtr = InnerProp->ContainerPtrToValuePtr<void>(ValuePtr);
-			StructObj->SetField(InnerProp->GetName(), DataTablePropertyToJsonValue(InnerProp, InnerValuePtr));
-		}
-		return MakeShared<FJsonValueObject>(StructObj);
-	}
-	else if (FArrayProperty* ArrayProp = CastField<FArrayProperty>(Property))
-	{
-		TArray<TSharedPtr<FJsonValue>> JsonArray;
-		FScriptArrayHelper ArrayHelper(ArrayProp, ValuePtr);
-		for (int32 i = 0; i < ArrayHelper.Num(); i++)
-		{
-			JsonArray.Add(DataTablePropertyToJsonValue(ArrayProp->Inner, ArrayHelper.GetRawPtr(i)));
-		}
-		return MakeShared<FJsonValueArray>(JsonArray);
-	}
-	else if (FMapProperty* MapProp = CastField<FMapProperty>(Property))
-	{
-		// Maps are exported as arrays of {key, value} objects
-		TArray<TSharedPtr<FJsonValue>> JsonArray;
-		FScriptMapHelper MapHelper(MapProp, ValuePtr);
-		for (int32 i = 0; i < MapHelper.Num(); i++)
-		{
-			if (MapHelper.IsValidIndex(i))
-			{
-				TSharedPtr<FJsonObject> EntryObj = MakeShared<FJsonObject>();
-				EntryObj->SetField(TEXT("key"), DataTablePropertyToJsonValue(MapProp->KeyProp, MapHelper.GetKeyPtr(i)));
-				EntryObj->SetField(TEXT("value"), DataTablePropertyToJsonValue(MapProp->ValueProp, MapHelper.GetValuePtr(i)));
-				JsonArray.Add(MakeShared<FJsonValueObject>(EntryObj));
-			}
-		}
-		return MakeShared<FJsonValueArray>(JsonArray);
-	}
-	else if (FSetProperty* SetProp = CastField<FSetProperty>(Property))
-	{
-		TArray<TSharedPtr<FJsonValue>> JsonArray;
-		FScriptSetHelper SetHelper(SetProp, ValuePtr);
-		for (int32 i = 0; i < SetHelper.Num(); i++)
-		{
-			if (SetHelper.IsValidIndex(i))
-			{
-				JsonArray.Add(DataTablePropertyToJsonValue(SetProp->ElementProp, SetHelper.GetElementPtr(i)));
-			}
-		}
-		return MakeShared<FJsonValueArray>(JsonArray);
-	}
-	else if (FSoftObjectProperty* SoftObjProp = CastField<FSoftObjectProperty>(Property))
-	{
-		const FSoftObjectPtr& SoftPtr = *static_cast<const FSoftObjectPtr*>(ValuePtr);
-		return MakeShared<FJsonValueString>(SoftPtr.ToString());
-	}
-	else if (FObjectProperty* ObjProp = CastField<FObjectProperty>(Property))
-	{
-		UObject* Obj = ObjProp->GetObjectPropertyValue(ValuePtr);
-		if (Obj)
-		{
-			return MakeShared<FJsonValueString>(Obj->GetPathName());
-		}
-		return MakeShared<FJsonValueString>(TEXT("None"));
-	}
-
-	// Fallback: export to string using UE's text export
-	FString StringValue;
-	Property->ExportTextItem_Direct(StringValue, ValuePtr, nullptr, nullptr, PPF_None);
-	return MakeShared<FJsonValueString>(StringValue);
-}
-
-//------------------------------------------------------------------------------
-// Helper: Convert a row to a JSON object
-//------------------------------------------------------------------------------
-static TSharedPtr<FJsonObject> RowToJsonObject(const UScriptStruct* RowStruct, const void* RowData)
-{
-	TSharedPtr<FJsonObject> RowObj = MakeShared<FJsonObject>();
-
-	for (TFieldIterator<FProperty> PropIt(RowStruct); PropIt; ++PropIt)
-	{
-		FProperty* Property = *PropIt;
-		const void* ValuePtr = Property->ContainerPtrToValuePtr<void>(RowData);
-		RowObj->SetField(Property->GetName(), DataTablePropertyToJsonValue(Property, ValuePtr));
-	}
-
-	return RowObj;
-}
-
-//------------------------------------------------------------------------------
-// Helper: Get property type as a friendly string
-//------------------------------------------------------------------------------
-static FString GetPropertyTypeString(FProperty* Property)
-{
-	if (FBoolProperty* BoolProp = CastField<FBoolProperty>(Property))
-	{
-		return TEXT("bool");
-	}
-	else if (FIntProperty* IntProp = CastField<FIntProperty>(Property))
-	{
-		return TEXT("int32");
-	}
-	else if (FInt64Property* Int64Prop = CastField<FInt64Property>(Property))
-	{
-		return TEXT("int64");
-	}
-	else if (FFloatProperty* FloatProp = CastField<FFloatProperty>(Property))
-	{
-		return TEXT("float");
-	}
-	else if (FDoubleProperty* DoubleProp = CastField<FDoubleProperty>(Property))
-	{
-		return TEXT("double");
-	}
-	else if (FStrProperty* StrProp = CastField<FStrProperty>(Property))
-	{
-		return TEXT("FString");
-	}
-	else if (FNameProperty* NameProp = CastField<FNameProperty>(Property))
-	{
-		return TEXT("FName");
-	}
-	else if (FTextProperty* TextProp = CastField<FTextProperty>(Property))
-	{
-		return TEXT("FText");
-	}
-	else if (FEnumProperty* EnumProp = CastField<FEnumProperty>(Property))
-	{
-		UEnum* Enum = EnumProp->GetEnum();
-		return FString::Printf(TEXT("enum:%s"), *Enum->GetName());
-	}
-	else if (FByteProperty* ByteProp = CastField<FByteProperty>(Property))
-	{
-		if (UEnum* Enum = ByteProp->GetIntPropertyEnum())
-		{
-			return FString::Printf(TEXT("enum:%s"), *Enum->GetName());
-		}
-		return TEXT("uint8");
-	}
-	else if (FStructProperty* StructProp = CastField<FStructProperty>(Property))
-	{
-		return FString::Printf(TEXT("struct:%s"), *StructProp->Struct->GetName());
-	}
-	else if (FArrayProperty* ArrayProp = CastField<FArrayProperty>(Property))
-	{
-		return FString::Printf(TEXT("array<%s>"), *GetPropertyTypeString(ArrayProp->Inner));
-	}
-	else if (FMapProperty* MapProp = CastField<FMapProperty>(Property))
-	{
-		return FString::Printf(TEXT("map<%s, %s>"), *GetPropertyTypeString(MapProp->KeyProp), *GetPropertyTypeString(MapProp->ValueProp));
-	}
-	else if (FSetProperty* SetProp = CastField<FSetProperty>(Property))
-	{
-		return FString::Printf(TEXT("set<%s>"), *GetPropertyTypeString(SetProp->ElementProp));
-	}
-	else if (FSoftObjectProperty* SoftObjProp = CastField<FSoftObjectProperty>(Property))
-	{
-		return FString::Printf(TEXT("soft_object:%s"), *SoftObjProp->PropertyClass->GetName());
-	}
-	else if (FObjectProperty* ObjProp = CastField<FObjectProperty>(Property))
-	{
-		return FString::Printf(TEXT("object:%s"), *ObjProp->PropertyClass->GetName());
-	}
-
-	return Property->GetCPPType();
-}
-
 //------------------------------------------------------------------------------
 // GetDataTableSchema
 //------------------------------------------------------------------------------
@@ -274,64 +47,22 @@ FECACommandResult FECACommand_GetDataTableSchema::Execute(const TSharedPtr<FJson
 	Result->SetStringField(TEXT("row_struct_path"), RowStruct->GetPathName());
 	Result->SetNumberField(TEXT("row_count"), DataTable->GetRowMap().Num());
 
-	// Build column definitions
-	TArray<TSharedPtr<FJsonValue>> ColumnsArray;
-	for (TFieldIterator<FProperty> PropIt(RowStruct); PropIt; ++PropIt)
-	{
-		FProperty* Property = *PropIt;
-		TSharedPtr<FJsonObject> ColumnObj = MakeShared<FJsonObject>();
-		ColumnObj->SetStringField(TEXT("name"), Property->GetName());
-		ColumnObj->SetStringField(TEXT("type"), GetPropertyTypeString(Property));
-		ColumnObj->SetStringField(TEXT("cpp_type"), Property->GetCPPType());
-
-		// Get tooltip/description if available
-		FString ToolTip = Property->GetMetaData(TEXT("ToolTip"));
-		if (!ToolTip.IsEmpty())
-		{
-			ColumnObj->SetStringField(TEXT("tooltip"), ToolTip);
-		}
+	// Serialize a default-initialized row to show the schema
+	// Uses the same serializer as get_datatable_rows so formats always match
+	uint8* DefaultData = (uint8*)FMemory::Malloc(RowStruct->GetStructureSize(), RowStruct->GetMinAlignment());
+	RowStruct->InitializeStruct(DefaultData);
 
-		// For enum properties, list the possible values
-		if (FEnumProperty* EnumProp = CastField<FEnumProperty>(Property))
-		{
-			UEnum* Enum = EnumProp->GetEnum();
-			TArray<TSharedPtr<FJsonValue>> EnumValues;
-			for (int32 i = 0; i < Enum->NumEnums() - 1; i++) // -1 to skip _MAX
-			{
-				EnumValues.Add(MakeShared<FJsonValueString>(Enum->GetNameStringByIndex(i)));
-			}
-			ColumnObj->SetArrayField(TEXT("enum_values"), EnumValues);
-		}
-		else if (FByteProperty* ByteProp = CastField<FByteProperty>(Property))
-		{
-			if (UEnum* Enum = ByteProp->GetIntPropertyEnum())
-			{
-				TArray<TSharedPtr<FJsonValue>> EnumValues;
-				for (int32 i = 0; i < Enum->NumEnums() - 1; i++)
-				{
-					EnumValues.Add(MakeShared<FJsonValueString>(Enum->GetNameStringByIndex(i)));
-				}
-				ColumnObj->SetArrayField(TEXT("enum_values"), EnumValues);
-			}
-		}
+	FString SchemaJson;
+	TSharedRef<TJsonWriter<>> JsonWriter = TJsonWriterFactory<>::Create(&SchemaJson);
+	JsonWriter->WriteObjectStart();
+	DataTable->WriteRowAsJSON(JsonWriter, DefaultData, EDataTableExportFlags::UseJsonObjectsForStructs);
+	JsonWriter->WriteObjectEnd();
+	JsonWriter->Close();
 
-		// For struct properties, list sub-fields
-		if (FStructProperty* StructProp = CastField<FStructProperty>(Property))
-		{
-			TArray<TSharedPtr<FJsonValue>> StructFields;
-			for (TFieldIterator<FProperty> InnerIt(StructProp->Struct); InnerIt; ++InnerIt)
-			{
-				TSharedPtr<FJsonObject> FieldInfo = MakeShared<FJsonObject>();
-				FieldInfo->SetStringField(TEXT("name"), InnerIt->GetName());
-				FieldInfo->SetStringField(TEXT("type"), GetPropertyTypeString(*InnerIt));
-				StructFields.Add(MakeShared<FJsonValueObject>(FieldInfo));
-			}
-			ColumnObj->SetArrayField(TEXT("struct_fields"), StructFields);
-		}
+	RowStruct->DestroyStruct(DefaultData);
+	FMemory::Free(DefaultData);
 
-		ColumnsArray.Add(MakeShared<FJsonValueObject>(ColumnObj));
-	}
-	Result->SetArrayField(TEXT("columns"), ColumnsArray);
+	Result->SetStringField(TEXT("default_row"), SchemaJson);
 
 	return FECACommandResult::Success(Result);
 }
@@ -360,15 +91,15 @@ FECACommandResult FECACommand_GetDataTableRows::Execute(const TSharedPtr<FJsonOb
 		return FECACommandResult::Error(TEXT("DataTable has no row structure defined"));
 	}
 
-	// Check for specific row names filter
+	// Check for row_names filter
 	const TArray<TSharedPtr<FJsonValue>>* RowNamesArray = nullptr;
-	TSet<FString> RequestedRowNames;
+	TSet<FName> RequestedRowNames;
 	bool bFilterByName = GetArrayParam(Params, TEXT("row_names"), RowNamesArray, false) && RowNamesArray && RowNamesArray->Num() > 0;
 	if (bFilterByName)
 	{
 		for (const TSharedPtr<FJsonValue>& Value : *RowNamesArray)
 		{
-			RequestedRowNames.Add(Value->AsString());
+			RequestedRowNames.Add(FName(*Value->AsString()));
 		}
 	}
 
@@ -378,23 +109,20 @@ FECACommandResult FECACommand_GetDataTableRows::Execute(const TSharedPtr<FJsonOb
 	GetIntParam(Params, TEXT("limit"), Limit, false);
 	GetIntParam(Params, TEXT("offset"), Offset, false);
 
-	TSharedPtr<FJsonObject> Result = MakeResult();
-	Result->SetStringField(TEXT("asset_path"), AssetPath);
-	Result->SetStringField(TEXT("row_struct_name"), RowStruct->GetName());
-
 	const TMap<FName, uint8*>& RowMap = DataTable->GetRowMap();
-	Result->SetNumberField(TEXT("total_row_count"), RowMap.Num());
 
-	TArray<TSharedPtr<FJsonValue>> RowsArray;
+	// Write filtered rows using engine's per-row JSON export (handles all property types)
+	FString JsonExport;
+	TSharedRef<TJsonWriter<>> JsonWriter = TJsonWriterFactory<>::Create(&JsonExport);
+	JsonWriter->WriteArrayStart();
+
 	int32 CurrentIndex = 0;
 	int32 ReturnedCount = 0;
 
 	for (const auto& Pair : RowMap)
 	{
-		FString RowName = Pair.Key.ToString();
-
 		// Apply name filter
-		if (bFilterByName && !RequestedRowNames.Contains(RowName))
+		if (bFilterByName && !RequestedRowNames.Contains(Pair.Key))
 		{
 			continue;
 		}
@@ -412,16 +140,24 @@ FECACommandResult FECACommand_GetDataTableRows::Execute(const TSharedPtr<FJsonOb
 			break;
 		}
 
-		TSharedPtr<FJsonObject> RowObj = MakeShared<FJsonObject>();
-		RowObj->SetStringField(TEXT("row_name"), RowName);
-		RowObj->SetObjectField(TEXT("data"), RowToJsonObject(RowStruct, Pair.Value));
+		// Write this row using engine's native JSON serialization
+		JsonWriter->WriteObjectStart();
+		JsonWriter->WriteValue(TEXT("Name"), Pair.Key.ToString());
+		DataTable->WriteRowAsJSON(JsonWriter, Pair.Value, EDataTableExportFlags::UseJsonObjectsForStructs);
+		JsonWriter->WriteObjectEnd();
 
-		RowsArray.Add(MakeShared<FJsonValueObject>(RowObj));
 		CurrentIndex++;
 		ReturnedCount++;
 	}
 
-	Result->SetArrayField(TEXT("rows"), RowsArray);
+	JsonWriter->WriteArrayEnd();
+	JsonWriter->Close();
+
+	TSharedPtr<FJsonObject> Result = MakeResult();
+	Result->SetStringField(TEXT("asset_path"), AssetPath);
+	Result->SetStringField(TEXT("row_struct_name"), RowStruct->GetName());
+	Result->SetNumberField(TEXT("total_row_count"), RowMap.Num());
+	Result->SetStringField(TEXT("json_export"), JsonExport);
 	Result->SetNumberField(TEXT("returned_count"), ReturnedCount);
 
 	return FECACommandResult::Success(Result);
@@ -550,7 +286,13 @@ FECACommandResult FECACommand_SetDataTableRow::Execute(const TSharedPtr<FJsonObj
 		uint8* StoredRow = DataTable->GetRowMap().FindRef(RowFName);
 		if (StoredRow)
 		{
-			Result->SetObjectField(TEXT("row_data"), RowToJsonObject(RowStruct, StoredRow));
+			FString RowJson;
+			TSharedRef<TJsonWriter<>> RowWriter = TJsonWriterFactory<>::Create(&RowJson);
+			RowWriter->WriteObjectStart();
+			DataTable->WriteRowAsJSON(RowWriter, StoredRow, EDataTableExportFlags::UseJsonObjectsForStructs);
+			RowWriter->WriteObjectEnd();
+			RowWriter->Close();
+			Result->SetStringField(TEXT("row_data"), RowJson);
 		}
 
 		return FECACommandResult::Success(Result);
@@ -641,7 +383,15 @@ FECACommandResult FECACommand_SetDataTableRow::Execute(const TSharedPtr<FJsonObj
 		}
 
 		// Return the updated row
-		Result->SetObjectField(TEXT("row_data"), RowToJsonObject(RowStruct, ExistingRowData));
+		{
+			FString RowJson;
+			TSharedRef<TJsonWriter<>> RowWriter = TJsonWriterFactory<>::Create(&RowJson);
+			RowWriter->WriteObjectStart();
+			DataTable->WriteRowAsJSON(RowWriter, ExistingRowData, EDataTableExportFlags::UseJsonObjectsForStructs);
+			RowWriter->WriteObjectEnd();
+			RowWriter->Close();
+			Result->SetStringField(TEXT("row_data"), RowJson);
+		}
 
 		return FECACommandResult::Success(Result);
 	}
diff --git a/Engine/Plugins/Experimental/ECABridge/Source/ECABridge/Public/Commands/ECADataTableCommands.h b/Engine/Plugins/Experimental/ECABridge/Source/ECABridge/Public/Commands/ECADataTableCommands.h
--- a/Engine/Plugins/Experimental/ECABridge/Source/ECABridge/Public/Commands/ECADataTableCommands.h
+++ b/Engine/Plugins/Experimental/ECABridge/Source/ECABridge/Public/Commands/ECADataTableCommands.h
@@ -11,7 +11,7 @@ class FECACommand_GetDataTableSchema : public IECACommand
 {
 public:
 	virtual FString GetName() const override { return TEXT("get_datatable_schema"); }
-	virtual FString GetDescription() const override { return TEXT("Get the row structure (schema) of a DataTable, including all column names and types"); }
+	virtual FString GetDescription() const override { return TEXT("Get the schema of a DataTable as a default-initialized row. Shows all fields with default values in the same format as get_datatable_rows. Result is in default_row string field."); }
 	virtual FString GetCategory() const override { return TEXT("DataTable"); }
 	
 	virtual TArray<FECACommandParam> GetParameters() const override
@@ -31,7 +31,7 @@ class FECACommand_GetDataTableRows : public IECACommand
 {
 public:
 	virtual FString GetName() const override { return TEXT("get_datatable_rows"); }
-	virtual FString GetDescription() const override { return TEXT("Get rows from a DataTable. Returns all rows or specific rows by name. Each row includes all column values."); }
+	virtual FString GetDescription() const override { return TEXT("Get rows from a DataTable. Returns all rows or specific rows by name. Result is in json_export string field."); }
 	virtual FString GetCategory() const override { return TEXT("DataTable"); }
 	
 	virtual TArray<FECACommandParam> GetParameters() const override
