[ECABridge] Add get_cvar MCP tool for reading console variable values

run_console_command can set CVars but cannot return their values.
Adds get_cvar tool that reads any CVar via IConsoleManager and returns
the name, current value, type (bool/int/float/string), and description.

#rnx
#Synced-CL 51740843

diff --git a/Engine/Plugins/Experimental/ECABridge/Source/ECABridge/Public/Commands/ECAEditorCommands.h b/Engine/Plugins/Experimental/ECABridge/Source/ECABridge/Public/Commands/ECAEditorCommands.h
--- a/Engine/Plugins/Experimental/ECABridge/Source/ECABridge/Public/Commands/ECAEditorCommands.h
+++ b/Engine/Plugins/Experimental/ECABridge/Source/ECABridge/Public/Commands/ECAEditorCommands.h
@@ -125,13 +125,34 @@ public:
 	virtual TArray<FECACommandParam> GetParameters() const override
 	{
 		return {
-			{ TEXT("command"), TEXT("string"), TEXT("The console command to execute"), true }
+			{ TEXT("command"), TEXT("string"), TEXT("The console command to execute"), true },
+			{ TEXT("target_world"), TEXT("string"), TEXT("Target world: 'auto' (PIE if playing, else editor), 'pie', 'editor'. Default: auto"), false }
 		};
 	}
 	
 	virtual FECACommandResult Execute(const TSharedPtr<FJsonObject>& Params) override;
 };
 
+/**
+ * Get a console variable value, type, and description
+ */
+class FECACommand_GetCVar : public IECACommand
+{
+public:
+	virtual FString GetName() const override { return TEXT("get_cvar"); }
+	virtual FString GetDescription() const override { return TEXT("Get a console variable value, type, and description"); }
+	virtual FString GetCategory() const override { return TEXT("Editor"); }
+
+	virtual TArray<FECACommandParam> GetParameters() const override
+	{
+		return {
+			{ TEXT("name"), TEXT("string"), TEXT("Console variable name (e.g. t.MaxFPS, r.ScreenPercentage)"), true }
+		};
+	}
+
+	virtual FECACommandResult Execute(const TSharedPtr<FJsonObject>& Params) override;
+};
+
 /**
  * Get current level info
  */
