[ECABridge] Fix static analysis warnings in MVVM debug commands

- C6263: Replace _alloca (FMemory_Alloca) with TArray<uint8> for
UFunction parameter buffer. Stack allocation in a loop accumulates
until function return; heap allocation frees each iteration.
- C2220/unreachable code: Replace nested ternary with if/else if/else
for function getter error messages. The else branch guaranteed
bSafeToInvoke was false, making the third ternary arm unreachable.

Pushing this through to fix build warnings.

#rnx
#Synced-CL 51757612

diff --git a/Engine/Plugins/Experimental/ECABridge/Source/ECABridge/Private/Commands/ECAMVVMDebugCommands.cpp b/Engine/Plugins/Experimental/ECABridge/Source/ECABridge/Private/Commands/ECAMVVMDebugCommands.cpp
--- a/Engine/Plugins/Experimental/ECABridge/Source/ECABridge/Private/Commands/ECAMVVMDebugCommands.cpp
+++ b/Engine/Plugins/Experimental/ECABridge/Source/ECABridge/Private/Commands/ECAMVVMDebugCommands.cpp
@@ -105,17 +105,17 @@ TSharedPtr<FJsonObject> FECACommand_GetViewmodelValues::SerializeViewModel(
 					&& (Func->HasAnyFunctionFlags(FUNC_Const) || Func->HasAnyFunctionFlags(FUNC_BlueprintPure));
 				if (bSafeToInvoke)
 				{
-					uint8* ParamBuffer = (uint8*)FMemory_Alloca(Func->ParmsSize);
-					FMemory::Memzero(ParamBuffer, Func->ParmsSize);
+					TArray<uint8> ParamBuffer;
+					ParamBuffer.SetNumZeroed(Func->ParmsSize);
 					// Initialize parameter properties only (CPF_Parm) — locals may have offsets beyond ParmsSize
 					for (TFieldIterator<FProperty> ParamIt(Func); ParamIt && (ParamIt->PropertyFlags & CPF_Parm); ++ParamIt)
 					{
-						ParamIt->InitializeValue_InContainer(ParamBuffer);
+						ParamIt->InitializeValue_InContainer(ParamBuffer.GetData());
 					}
-					VM->ProcessEvent(Func, ParamBuffer);
+					VM->ProcessEvent(Func, ParamBuffer.GetData());
 
 					FString ReturnStr;
-					const void* ReturnPtr = ReturnProp->ContainerPtrToValuePtr<void>(ParamBuffer);
+					const void* ReturnPtr = ReturnProp->ContainerPtrToValuePtr<void>(ParamBuffer.GetData());
 					ReturnProp->ExportText_Direct(ReturnStr, ReturnPtr, ReturnPtr, nullptr, PPF_None);
 					FuncObj->SetStringField(TEXT("value"), ReturnStr);
 					FuncObj->SetStringField(TEXT("return_type"), ReturnProp->GetCPPType());
@@ -123,15 +123,16 @@ TSharedPtr<FJsonObject> FECACommand_GetViewmodelValues::SerializeViewModel(
 					// Destroy parameter properties only
 					for (TFieldIterator<FProperty> ParamIt(Func); ParamIt && (ParamIt->PropertyFlags & CPF_Parm); ++ParamIt)
 					{
-						ParamIt->DestroyValue_InContainer(ParamBuffer);
+						ParamIt->DestroyValue_InContainer(ParamBuffer.GetData());
 					}
 				}
+				else if (!ReturnProp)
+				{
+					FuncObj->SetStringField(TEXT("value"), TEXT("<no return value>"));
+				}
 				else
 				{
-					FuncObj->SetStringField(TEXT("value"),
-						!ReturnProp ? TEXT("<no return value>") :
-						!bSafeToInvoke ? TEXT("<not const/pure - skipped for safety>") :
-						TEXT("<function with params - cannot auto-invoke>"));
+					FuncObj->SetStringField(TEXT("value"), TEXT("<not const/pure - skipped for safety>"));
 				}
 
 				FuncGettersArray.Add(MakeShared<FJsonValueObject>(FuncObj));
