Fix V595 static analysis warning in ECAEditorCommands.cpp

Guard GEditor dereference with null check before accessing PlayWorld.
The pointer was used on line 616 before being verified against nullptr on line 627.

#rnx
#Synced-CL 51871011
#Synced-CL 51871011

diff --git a/Engine/Plugins/Experimental/ECABridge/Source/ECABridge/Private/Commands/ECAEditorCommands.cpp b/Engine/Plugins/Experimental/ECABridge/Source/ECABridge/Private/Commands/ECAEditorCommands.cpp
--- a/Engine/Plugins/Experimental/ECABridge/Source/ECABridge/Private/Commands/ECAEditorCommands.cpp
+++ b/Engine/Plugins/Experimental/ECABridge/Source/ECABridge/Private/Commands/ECAEditorCommands.cpp
@@ -613,10 +613,13 @@ FECACommandResult FECACommand_RunConsoleCommand::Execute(const TSharedPtr<FJsonO
 
 	if (bUsePIE)
 	{
-		World = GEditor->PlayWorld;
-		if (!GIsPlayInEditorWorld)
+		if (GEditor)
 		{
-			OldWorld = SetPlayInEditorWorld(World);
+			World = GEditor->PlayWorld;
+			if (!GIsPlayInEditorWorld)
+			{
+				OldWorld = SetPlayInEditorWorld(World);
+			}
 		}
 	}
 	else
