ECABridge: Add get_references and get_referencers commands

- get_references: query what an asset depends on (hard/soft refs, searchable names)
- get_referencers: query what depends on an asset or gameplay tag
Both use FAssetIdentifier to return package deps AND searchable names
(gameplay tags, FPrimaryAssetId, etc.). get_referencers accepts a tag parameter
to find all assets referencing a gameplay tag.

#tests Ran tests on editor data
#Synced-CL 51745444

diff --git a/Engine/Plugins/Experimental/ECABridge/Source/ECABridge/ECABridge.Build.cs b/Engine/Plugins/Experimental/ECABridge/Source/ECABridge/ECABridge.Build.cs
--- a/Engine/Plugins/Experimental/ECABridge/Source/ECABridge/ECABridge.Build.cs
+++ b/Engine/Plugins/Experimental/ECABridge/Source/ECABridge/ECABridge.Build.cs
@@ -22,6 +22,8 @@ public class ECABridge : ModuleRules
 		{
 			// Editor core
 			"UnrealEd",
+			// Gameplay tags (for FGameplayTag::StaticStruct in get_referencers)
+			"GameplayTags",
 			
 			// MetaSound
 			"MetasoundEngine", "MetasoundEditor", "MetasoundFrontend", "MetasoundGraphCore", "MetasoundStandardNodes", "AudioExtensions",
diff --git a/Engine/Plugins/Experimental/ECABridge/Source/ECABridge/Private/Commands/ECAAssetCommands.cpp b/Engine/Plugins/Experimental/ECABridge/Source/ECABridge/Private/Commands/ECAAssetCommands.cpp
--- a/Engine/Plugins/Experimental/ECABridge/Source/ECABridge/Private/Commands/ECAAssetCommands.cpp
+++ b/Engine/Plugins/Experimental/ECABridge/Source/ECABridge/Private/Commands/ECAAssetCommands.cpp
@@ -49,6 +49,7 @@
 #include "Misc/Base64.h"
 #include "CanvasTypes.h"
 #include "CanvasItem.h"
+#include "GameplayTagContainer.h"
 
 // Register all asset commands
 REGISTER_ECA_COMMAND(FECACommand_ImportTexture)
@@ -79,6 +80,8 @@ REGISTER_ECA_COMMAND(FECACommand_MoveAsset)
 REGISTER_ECA_COMMAND(FECACommand_DuplicateAsset)
 REGISTER_ECA_COMMAND(FECACommand_GetAssetThumbnail)
 REGISTER_ECA_COMMAND(FECACommand_GetAssetThumbnails)
+REGISTER_ECA_COMMAND(FECACommand_GetReferences)
+REGISTER_ECA_COMMAND(FECACommand_GetReferencers)
 
 //------------------------------------------------------------------------------
 // Helper functions
@@ -3779,3 +3782,195 @@ FECACommandResult FECACommand_GetAssetThumbnails::Execute(const TSharedPtr<FJson
 
 	return FECACommandResult::Success(Result);
 }
+
+
+//------------------------------------------------------------------------------
+// Asset References / Referencers
+//------------------------------------------------------------------------------
+
+static UE::AssetRegistry::EDependencyCategory ParseDependencyCategory(const FString& CategoryStr)
+{
+	if (CategoryStr.Equals(TEXT("package"), ESearchCase::IgnoreCase))
+	{
+		return UE::AssetRegistry::EDependencyCategory::Package;
+	}
+	else if (CategoryStr.Equals(TEXT("searchable_name"), ESearchCase::IgnoreCase))
+	{
+		return UE::AssetRegistry::EDependencyCategory::SearchableName;
+	}
+	else if (CategoryStr.Equals(TEXT("manage"), ESearchCase::IgnoreCase))
+	{
+		return UE::AssetRegistry::EDependencyCategory::Manage;
+	}
+	return UE::AssetRegistry::EDependencyCategory::All;
+}
+
+static FAssetIdentifier AssetPathToIdentifier(const FString& AssetPath)
+{
+	// Strip object name suffix if present (e.g., "/Game/Data/Table.Table" -> "/Game/Data/Table")
+	int32 DotIndex;
+	if (AssetPath.FindChar(TEXT('.'), DotIndex))
+	{
+		return FAssetIdentifier(FName(*AssetPath.Left(DotIndex)));
+	}
+	return FAssetIdentifier(FName(*AssetPath));
+}
+
+static TSharedPtr<FJsonObject> AssetIdentifierToJson(const FAssetIdentifier& Identifier, IAssetRegistry& Registry)
+{
+	TSharedPtr<FJsonObject> Entry = MakeShared<FJsonObject>();
+
+	if (Identifier.IsPackage())
+	{
+		FName PackageName = Identifier.PackageName;
+		Entry->SetStringField(TEXT("type"), TEXT("package"));
+		Entry->SetStringField(TEXT("package_name"), PackageName.ToString());
+
+		TArray<FAssetData> Assets;
+		Registry.GetAssetsByPackageName(PackageName, Assets, true);
+
+		if (Assets.Num() > 0)
+		{
+			const FAssetData& Asset = Assets[0];
+			Entry->SetStringField(TEXT("asset_path"), Asset.GetObjectPathString());
+			Entry->SetStringField(TEXT("asset_name"), Asset.AssetName.ToString());
+			Entry->SetStringField(TEXT("class"), Asset.AssetClassPath.GetAssetName().ToString());
+		}
+	}
+	else if (Identifier.IsValue())
+	{
+		// Searchable name entry (gameplay tags, FPrimaryAssetId, etc.)
+		Entry->SetStringField(TEXT("type"), TEXT("searchable_name"));
+		Entry->SetStringField(TEXT("name"), Identifier.ValueName.ToString());
+
+		if (Identifier.ObjectName != NAME_None)
+		{
+			Entry->SetStringField(TEXT("object_name"), Identifier.ObjectName.ToString());
+		}
+	}
+	else
+	{
+		Entry->SetStringField(TEXT("type"), TEXT("other"));
+		Entry->SetStringField(TEXT("identifier"), Identifier.ToString());
+	}
+
+	return Entry;
+}
+
+FECACommandResult FECACommand_GetReferences::Execute(const TSharedPtr<FJsonObject>& Params)
+{
+	FString AssetPath;
+	if (!GetStringParam(Params, TEXT("asset_path"), AssetPath))
+	{
+		return FECACommandResult::Error(TEXT("Missing required parameter: asset_path"));
+	}
+
+	FString CategoryStr = TEXT("all");
+	GetStringParam(Params, TEXT("category"), CategoryStr, false);
+
+	bool bHardOnly = false;
+	GetBoolParam(Params, TEXT("hard_only"), bHardOnly, false);
+
+	int32 Limit = 100;
+	GetIntParam(Params, TEXT("limit"), Limit, false);
+
+	IAssetRegistry& Registry = FModuleManager::LoadModuleChecked<FAssetRegistryModule>("AssetRegistry").Get();
+
+	FAssetIdentifier SourceIdentifier = AssetPathToIdentifier(AssetPath);
+
+	UE::AssetRegistry::EDependencyCategory Category = ParseDependencyCategory(CategoryStr);
+	UE::AssetRegistry::FDependencyQuery Flags;
+	if (bHardOnly && Category == UE::AssetRegistry::EDependencyCategory::Package)
+	{
+		Flags.Required = UE::AssetRegistry::EDependencyProperty::Hard;
+	}
+
+	TArray<FAssetIdentifier> OutDependencies;
+	Registry.GetDependencies(SourceIdentifier, OutDependencies, Category, Flags);
+
+	int32 TotalFound = OutDependencies.Num();
+	int32 ReturnCount = FMath::Min(TotalFound, Limit);
+
+	TArray<TSharedPtr<FJsonValue>> ResultArray;
+	for (int32 i = 0; i < ReturnCount; ++i)
+	{
+		ResultArray.Add(MakeShared<FJsonValueObject>(AssetIdentifierToJson(OutDependencies[i], Registry)));
+	}
+
+	TSharedPtr<FJsonObject> Result = MakeShared<FJsonObject>();
+	Result->SetStringField(TEXT("asset_path"), AssetPath);
+	Result->SetNumberField(TEXT("total_found"), TotalFound);
+	Result->SetNumberField(TEXT("returned"), ReturnCount);
+	Result->SetStringField(TEXT("category"), CategoryStr);
+	Result->SetArrayField(TEXT("dependencies"), ResultArray);
+
+	return FECACommandResult::Success(Result);
+}
+
+FECACommandResult FECACommand_GetReferencers::Execute(const TSharedPtr<FJsonObject>& Params)
+{
+	FString AssetPath;
+	bool bHasAssetPath = GetStringParam(Params, TEXT("asset_path"), AssetPath, false);
+
+	FString TagName;
+	bool bHasTag = GetStringParam(Params, TEXT("tag"), TagName, false);
+
+	if (bHasAssetPath && bHasTag)
+	{
+		return FECACommandResult::Error(TEXT("Cannot specify both asset_path and tag. Use one or the other."));
+	}
+	if (!bHasAssetPath && !bHasTag)
+	{
+		return FECACommandResult::Error(TEXT("Must specify either asset_path or tag. Use asset_path for asset referencers, or tag to find all assets that reference a gameplay tag."));
+	}
+
+	FString CategoryStr = TEXT("all");
+	GetStringParam(Params, TEXT("category"), CategoryStr, false);
+
+	int32 Limit = 100;
+	GetIntParam(Params, TEXT("limit"), Limit, false);
+
+	IAssetRegistry& Registry = FModuleManager::LoadModuleChecked<FAssetRegistryModule>("AssetRegistry").Get();
+
+	FAssetIdentifier SourceIdentifier;
+	if (bHasTag)
+	{
+		// Construct identifier for a gameplay tag searchable name
+		SourceIdentifier = FAssetIdentifier(FGameplayTag::StaticStruct(), FName(*TagName));
+	}
+	else
+	{
+		SourceIdentifier = AssetPathToIdentifier(AssetPath);
+	}
+
+	UE::AssetRegistry::EDependencyCategory Category = ParseDependencyCategory(CategoryStr);
+	UE::AssetRegistry::FDependencyQuery Flags;
+
+	TArray<FAssetIdentifier> OutReferencers;
+	Registry.GetReferencers(SourceIdentifier, OutReferencers, Category, Flags);
+
+	int32 TotalFound = OutReferencers.Num();
+	int32 ReturnCount = FMath::Min(TotalFound, Limit);
+
+	TArray<TSharedPtr<FJsonValue>> ResultArray;
+	for (int32 i = 0; i < ReturnCount; ++i)
+	{
+		ResultArray.Add(MakeShared<FJsonValueObject>(AssetIdentifierToJson(OutReferencers[i], Registry)));
+	}
+
+	TSharedPtr<FJsonObject> Result = MakeShared<FJsonObject>();
+	if (bHasTag)
+	{
+		Result->SetStringField(TEXT("tag"), TagName);
+	}
+	else
+	{
+		Result->SetStringField(TEXT("asset_path"), AssetPath);
+	}
+	Result->SetNumberField(TEXT("total_found"), TotalFound);
+	Result->SetNumberField(TEXT("returned"), ReturnCount);
+	Result->SetStringField(TEXT("category"), CategoryStr);
+	Result->SetArrayField(TEXT("referencers"), ResultArray);
+
+	return FECACommandResult::Success(Result);
+}
diff --git a/Engine/Plugins/Experimental/ECABridge/Source/ECABridge/Public/Commands/ECAAssetCommands.h b/Engine/Plugins/Experimental/ECABridge/Source/ECABridge/Public/Commands/ECAAssetCommands.h
--- a/Engine/Plugins/Experimental/ECABridge/Source/ECABridge/Public/Commands/ECAAssetCommands.h
+++ b/Engine/Plugins/Experimental/ECABridge/Source/ECABridge/Public/Commands/ECAAssetCommands.h
@@ -624,6 +624,52 @@ public:
 	virtual FECACommandResult Execute(const TSharedPtr<FJsonObject>& Params) override;
 };
 
+/**
+ * Get the dependencies of an asset (what it references)
+ */
+class FECACommand_GetReferences : public IECACommand
+{
+public:
+	virtual FString GetName() const override { return TEXT("get_references"); }
+	virtual FString GetDescription() const override { return TEXT("Get what an asset depends on (its references/dependencies). Returns hard refs, soft refs, and searchable names (gameplay tags). Use get_referencers for the reverse direction."); }
+	virtual FString GetCategory() const override { return TEXT("Asset"); }
+	
+	virtual TArray<FECACommandParam> GetParameters() const override
+	{
+		return {
+			{ TEXT("asset_path"), TEXT("string"), TEXT("Content path to the asset (e.g., /Game/Data/MyDataTable, /Game/Blueprints/BP_MyActor)"), true },
+			{ TEXT("category"), TEXT("string"), TEXT("Filter by dependency category: 'package' (hard/soft asset refs), 'searchable_name' (gameplay tags), 'manage' (asset manager), or 'all'"), false, TEXT("all") },
+			{ TEXT("hard_only"), TEXT("boolean"), TEXT("When category is 'package', only return hard dependencies (must-load-before). Excludes soft references."), false, TEXT("false") },
+			{ TEXT("limit"), TEXT("number"), TEXT("Max results to return"), false, TEXT("100") }
+		};
+	}
+	
+	virtual FECACommandResult Execute(const TSharedPtr<FJsonObject>& Params) override;
+};
+
+/**
+ * Get what references an asset or gameplay tag (what depends on it)
+ */
+class FECACommand_GetReferencers : public IECACommand
+{
+public:
+	virtual FString GetName() const override { return TEXT("get_referencers"); }
+	virtual FString GetDescription() const override { return TEXT("Get what references an asset or gameplay tag. For assets: shows what would break if deleted. For gameplay tags: finds all assets that use a tag (e.g., 'GameplayCue.MyGame.MyEffect'). Pass either asset_path or tag, not both."); }
+	virtual FString GetCategory() const override { return TEXT("Asset"); }
+	
+	virtual TArray<FECACommandParam> GetParameters() const override
+	{
+		return {
+			{ TEXT("asset_path"), TEXT("string"), TEXT("Content path to the asset (e.g., /Game/Data/MyDataTable). Mutually exclusive with tag."), false },
+			{ TEXT("tag"), TEXT("string"), TEXT("Gameplay tag to search for (e.g., 'GameplayCue.MyGame.MyEffect', 'Ability.MyGame.MyAbility'). Finds all assets that reference this tag. Mutually exclusive with asset_path."), false },
+			{ TEXT("category"), TEXT("string"), TEXT("Filter by dependency category: 'package' (hard/soft asset refs), 'searchable_name' (gameplay tags), 'manage' (asset manager), or 'all'"), false, TEXT("all") },
+			{ TEXT("limit"), TEXT("number"), TEXT("Max results to return"), false, TEXT("100") }
+		};
+	}
+	
+	virtual FECACommandResult Execute(const TSharedPtr<FJsonObject>& Params) override;
+};
+
 /**
  * Get thumbnails for multiple assets in a single call.
  * IMPORTANT FOR AI: Use this for bulk visual analysis of assets. For each thumbnail, use analyze_image()
