[ECABridge] Clear dirty flag on temp package in diff_widget_blueprint

LoadPackageForDiff loads the depot revision into a temp package for
comparison. Without clearing the dirty flag before marking as garbage,
PIE prompts to save the temp package on next play.

#rnx
#Synced-CL 51881091

diff --git a/Engine/Plugins/Experimental/ECABridge/Source/ECABridge/Private/Commands/ECADiffCommands.cpp b/Engine/Plugins/Experimental/ECABridge/Source/ECABridge/Private/Commands/ECADiffCommands.cpp
--- a/Engine/Plugins/Experimental/ECABridge/Source/ECABridge/Private/Commands/ECADiffCommands.cpp
+++ b/Engine/Plugins/Experimental/ECABridge/Source/ECABridge/Private/Commands/ECADiffCommands.cpp
@@ -488,9 +488,10 @@ FECACommandResult FECACommand_DiffWidgetBlueprint::Execute(const TSharedPtr<FJso
 		ResultObj->SetObjectField(TEXT("variables"), DiffVariables(OldWBP, NewWBP));
 	}
 
-	// Cleanup temp package
+	// Cleanup temp package - must clear dirty flag or PIE will try to save it
 	if (OldPackage)
 	{
+		OldPackage->SetDirtyFlag(false);
 		OldPackage->ClearFlags(RF_Standalone);
 		OldPackage->MarkAsGarbage();
 	}
