ECABridge: EDA/ECA Collaborative Interop

#Synced-CL 51761123

diff --git a/Engine/Plugins/Experimental/ECABridge/ECABridge.uplugin b/Engine/Plugins/Experimental/ECABridge/ECABridge.uplugin
--- a/Engine/Plugins/Experimental/ECABridge/ECABridge.uplugin
+++ b/Engine/Plugins/Experimental/ECABridge/ECABridge.uplugin
@@ -23,6 +23,10 @@
 		}
 	],
 	"Plugins": [
+		{
+			"Name": "ToolsetRegistry",
+			"Enabled": true
+		},
 		{
 			"Name": "Niagara",
 			"Enabled": true
diff --git a/Engine/Plugins/Experimental/ECABridge/Source/ECABridge/ECABridge.Build.cs b/Engine/Plugins/Experimental/ECABridge/Source/ECABridge/ECABridge.Build.cs
--- a/Engine/Plugins/Experimental/ECABridge/Source/ECABridge/ECABridge.Build.cs
+++ b/Engine/Plugins/Experimental/ECABridge/Source/ECABridge/ECABridge.Build.cs
@@ -59,6 +59,9 @@ public class ECABridge : ModuleRules
 			// Additional utilities
 			"Projects", "DeveloperSettings", "ToolMenus",
 			
+			// ToolsetRegistry integration (bidirectional: expose ECA commands as toolset + consume toolsets via MCP)
+			"ToolsetRegistry",
+			
 			// Image processing
 			"ImageWrapper",
 			
diff --git a/Engine/Plugins/Experimental/ECABridge/Source/ECABridge/Private/ECABridgeModule.cpp b/Engine/Plugins/Experimental/ECABridge/Source/ECABridge/Private/ECABridgeModule.cpp
--- a/Engine/Plugins/Experimental/ECABridge/Source/ECABridge/Private/ECABridgeModule.cpp
+++ b/Engine/Plugins/Experimental/ECABridge/Source/ECABridge/Private/ECABridgeModule.cpp
@@ -3,7 +3,9 @@
 #include "ECABridgeModule.h"
 #include "ECABridge.h"
 #include "ECABridgeSettings.h"
+#include "ECABridgeToolset.h"
 #include "ECAEventQueue.h"
+#include "ToolsetRegistry/ToolsetRegistrySubsystem.h"
 #include "HAL/IConsoleManager.h"
 #include "Editor.h"
 #include "Selection.h"
@@ -41,11 +43,22 @@ void FECABridgeModule::StartupModule()
 	// Register editor event handlers
 	RegisterEditorEventHandlers();
 	
+	// Register ECABridge commands with the ToolsetRegistry (Direction 1)
+	// Deferred until GEditor is available since ToolsetRegistrySubsystem is an editor subsystem
+	if (GEditor)
+	{
+		RegisterWithToolsetRegistry();
+	}
+	// else: handled in OnEditorInitialized/RegisterGEditorDependentHandlers
+	
 	UE_LOG(LogTemp, Log, TEXT("[ECABridge] Module ready"));
 }
 
 void FECABridgeModule::ShutdownModule()
 {
+	// Unregister from ToolsetRegistry
+	UnregisterFromToolsetRegistry();
+	
 	// Unregister editor event handlers
 	UnregisterEditorEventHandlers();
 	
@@ -173,6 +186,12 @@ void FECABridgeModule::RegisterGEditorDependentHandlers()
 	CameraMovedHandle = FEditorDelegates::OnEditorCameraMoved.AddRaw(this, &FECABridgeModule::OnEditorCameraMoved);
 	UE_LOG(LogTemp, Log, TEXT("[ECABridge] Registered camera movement handler"));
 	
+	// Register with ToolsetRegistry now that GEditor is available (if not already done)
+	if (!ECABridgeToolset.IsValid())
+	{
+		RegisterWithToolsetRegistry();
+	}
+	
 	UE_LOG(LogTemp, Log, TEXT("[ECABridge] All GEditor-dependent handlers registered"));
 }
 
@@ -273,6 +292,58 @@ void FECABridgeModule::UnregisterEditorEventHandlers()
 	UE_LOG(LogTemp, Log, TEXT("[ECABridge] All editor event handlers unregistered"));
 }
 
+void FECABridgeModule::RegisterWithToolsetRegistry()
+{
+	if (!GEditor)
+	{
+		UE_LOG(LogTemp, Warning, TEXT("[ECABridge] Cannot register with ToolsetRegistry: GEditor not available"));
+		return;
+	}
+	
+	UToolsetRegistrySubsystem* ToolsetRegistrySubsystem = GEditor->GetEditorSubsystem<UToolsetRegistrySubsystem>();
+	if (!ToolsetRegistrySubsystem)
+	{
+		UE_LOG(LogTemp, Warning, TEXT("[ECABridge] ToolsetRegistrySubsystem not available — ToolsetRegistry plugin may not be enabled"));
+		return;
+	}
+	
+	// Direction 1: Register ECABridge commands as a toolset
+	ECABridgeToolset = MakeShared<FECABridgeToolset>();
+	if (ToolsetRegistrySubsystem->ToolsetRegistry.RegisterToolset(ECABridgeToolset))
+	{
+		int32 CommandCount = FECACommandRegistry::Get().GetAllCommands().Num();
+		UE_LOG(LogTemp, Log, TEXT("[ECABridge] ✓ Registered %d ECABridge commands with ToolsetRegistry as toolset '%s'"),
+			CommandCount, *ECABridgeToolset->GetToolsetName());
+	}
+	else
+	{
+		UE_LOG(LogTemp, Warning, TEXT("[ECABridge] Failed to register ECABridgeToolset with ToolsetRegistry"));
+		ECABridgeToolset.Reset();
+	}
+}
+
+void FECABridgeModule::UnregisterFromToolsetRegistry()
+{
+	if (ECABridgeToolset.IsValid() && GEditor)
+	{
+		if (UToolsetRegistrySubsystem* ToolsetRegistrySubsystem = GEditor->GetEditorSubsystem<UToolsetRegistrySubsystem>())
+		{
+			ToolsetRegistrySubsystem->ToolsetRegistry.UnregisterToolset(ECABridgeToolset);
+			UE_LOG(LogTemp, Log, TEXT("[ECABridge] Unregistered ECABridgeToolset from ToolsetRegistry"));
+		}
+		ECABridgeToolset.Reset();
+	}
+	
+	if (ToolsetRegistryChangedHandle.IsValid() && GEditor)
+	{
+		if (UToolsetRegistrySubsystem* ToolsetRegistrySubsystem = GEditor->GetEditorSubsystem<UToolsetRegistrySubsystem>())
+		{
+			ToolsetRegistrySubsystem->ToolsetRegistry.OnToolsetRegistered().Remove(ToolsetRegistryChangedHandle);
+		}
+		ToolsetRegistryChangedHandle.Reset();
+	}
+}
+
 void FECABridgeModule::OnEditorSelectionChanged(UObject* SelectionObject)
 {
 	UE_LOG(LogTemp, Log, TEXT("[ECABridge] OnEditorSelectionChanged called! Object: %s"), 
diff --git a/Engine/Plugins/Experimental/ECABridge/Source/ECABridge/Private/ECABridgeToolset.cpp b/Engine/Plugins/Experimental/ECABridge/Source/ECABridge/Private/ECABridgeToolset.cpp
new file mode 100644
--- /dev/null
+++ b/Engine/Plugins/Experimental/ECABridge/Source/ECABridge/Private/ECABridgeToolset.cpp
@@ -0,0 +1,184 @@
+// Copyright Epic Games, Inc. All Rights Reserved.
+
+#include "ECABridgeToolset.h"
+#include "Commands/ECACommand.h"
+#include "Dom/JsonObject.h"
+#include "Dom/JsonValue.h"
+#include "Serialization/JsonReader.h"
+#include "Serialization/JsonWriter.h"
+#include "Serialization/JsonSerializer.h"
+#include "Async/Async.h"
+
+FECABridgeToolset::FECABridgeToolset()
+{
+}
+
+FECABridgeToolset::~FECABridgeToolset()
+{
+}
+
+FString FECABridgeToolset::GetToolsetName() const
+{
+	return TEXT("ECABridge");
+}
+
+FString FECABridgeToolset::GetToolsetVersion() const
+{
+	return TEXT("1.0");
+}
+
+TFuture<TValueOrError<FString, FString>> FECABridgeToolset::ExecuteTool(
+	const FString& ToolName, const FString& JsonInput)
+{
+	// ECABridge commands must run on the game thread since they interact with editor subsystems.
+	// If we're already on the game thread, execute directly and return a fulfilled future.
+	// Otherwise, marshal to the game thread via Async.
+
+	auto DoExecute = [ToolName, JsonInput]() -> TValueOrError<FString, FString>
+	{
+		// Parse the JSON input into an FJsonObject for the command
+		TSharedPtr<FJsonObject> Params;
+		if (!JsonInput.IsEmpty() && JsonInput != TEXT("{}"))
+		{
+			TSharedRef<TJsonReader<>> Reader = TJsonReaderFactory<>::Create(JsonInput);
+			if (!FJsonSerializer::Deserialize(Reader, Params) || !Params.IsValid())
+			{
+				return MakeError(FString::Printf(TEXT("Failed to parse JSON input for tool '%s'"), *ToolName));
+			}
+		}
+		else
+		{
+			Params = MakeShared<FJsonObject>();
+		}
+
+		// Execute via the ECABridge command registry
+		FECACommandResult Result = FECACommandRegistry::Get().ExecuteCommand(ToolName, Params);
+
+		if (Result.bSuccess)
+		{
+			// Serialize the result data to JSON string
+			if (Result.ResultData.IsValid())
+			{
+				FString OutputString;
+				TSharedRef<TJsonWriter<>> Writer = TJsonWriterFactory<>::Create(&OutputString);
+				FJsonSerializer::Serialize(Result.ResultData.ToSharedRef(), Writer);
+				return MakeValue(MoveTemp(OutputString));
+			}
+			else
+			{
+				return MakeValue(FString(TEXT("{\"success\": true}")));
+			}
+		}
+		else
+		{
+			return MakeError(Result.ErrorMessage);
+		}
+	};
+
+	if (IsInGameThread())
+	{
+		return MakeFulfilledPromise<TValueOrError<FString, FString>>(DoExecute()).GetFuture();
+	}
+	else
+	{
+		return Async(EAsyncExecution::TaskGraphMainThread, MoveTemp(DoExecute));
+	}
+}
+
+FString FECABridgeToolset::GetToolsetJsonSchema() const
+{
+	// Build a JSON array of tool definitions from all registered ECABridge commands.
+	// The ToolsetRegistry expects a JSON array where each element describes a toolset
+	// with a "tools" array inside it.
+	//
+	// Format:
+	// [
+	//   {
+	//     "name": "ECABridge",
+	//     "description": "...",
+	//     "tools": [
+	//       { "name": "create_actor", "description": "...", "inputSchema": {...} },
+	//       ...
+	//     ]
+	//   }
+	// ]
+
+	TArray<TSharedPtr<IECACommand>> Commands = FECACommandRegistry::Get().GetAllCommands();
+
+	TArray<TSharedPtr<FJsonValue>> ToolsArray;
+	for (const TSharedPtr<IECACommand>& Command : Commands)
+	{
+		TSharedPtr<FJsonObject> ToolDef = MakeShared<FJsonObject>();
+		ToolDef->SetStringField(TEXT("name"), Command->GetName());
+		ToolDef->SetStringField(TEXT("description"), Command->GetDescription());
+
+		// Build inputSchema from command parameters
+		TSharedPtr<FJsonObject> InputSchema = MakeShared<FJsonObject>();
+		InputSchema->SetStringField(TEXT("type"), TEXT("object"));
+
+		TSharedPtr<FJsonObject> Properties = MakeShared<FJsonObject>();
+		TArray<TSharedPtr<FJsonValue>> Required;
+
+		for (const FECACommandParam& Param : Command->GetParameters())
+		{
+			TSharedPtr<FJsonObject> PropDef = MakeShared<FJsonObject>();
+
+			// Map ECA param types to JSON Schema types
+			FString JsonSchemaType = Param.Type;
+			if (Param.Type == TEXT("float") || Param.Type == TEXT("double"))
+			{
+				JsonSchemaType = TEXT("number");
+			}
+			else if (Param.Type == TEXT("int") || Param.Type == TEXT("int32") || Param.Type == TEXT("int64"))
+			{
+				JsonSchemaType = TEXT("integer");
+			}
+			else if (Param.Type == TEXT("vector") || Param.Type == TEXT("rotator") || Param.Type == TEXT("transform"))
+			{
+				JsonSchemaType = TEXT("object");
+			}
+			else if (Param.Type == TEXT("any"))
+			{
+				JsonSchemaType = TEXT("string");
+			}
+
+			PropDef->SetStringField(TEXT("type"), JsonSchemaType);
+			PropDef->SetStringField(TEXT("description"), Param.Description);
+
+			if (!Param.DefaultValue.IsEmpty())
+			{
+				PropDef->SetStringField(TEXT("default"), Param.DefaultValue);
+			}
+
+			Properties->SetObjectField(Param.Name, PropDef);
+
+			if (Param.bRequired)
+			{
+				Required.Add(MakeShared<FJsonValueString>(Param.Name));
+			}
+		}
+
+		InputSchema->SetObjectField(TEXT("properties"), Properties);
+		if (Required.Num() > 0)
+		{
+			InputSchema->SetArrayField(TEXT("required"), Required);
+		}
+
+		ToolDef->SetObjectField(TEXT("inputSchema"), InputSchema);
+		ToolsArray.Add(MakeShared<FJsonValueObject>(ToolDef));
+	}
+
+	// Build the toolset envelope — a single JSON object (not an array).
+	// The registry's GetToolsetJsonSchemas() wraps each toolset's schema into an array.
+	TSharedPtr<FJsonObject> ToolsetEnvelope = MakeShared<FJsonObject>();
+	ToolsetEnvelope->SetStringField(TEXT("name"), GetToolsetName());
+	ToolsetEnvelope->SetStringField(TEXT("version"), GetToolsetVersion());
+	ToolsetEnvelope->SetStringField(TEXT("description"), 
+		TEXT("ECA Bridge editor commands — 229+ tools for actors, Blueprints, assets, materials, Niagara, MetaSound, UMG, mesh operations, and more"));
+	ToolsetEnvelope->SetArrayField(TEXT("tools"), ToolsArray);
+
+	FString OutputString;
+	TSharedRef<TJsonWriter<>> Writer = TJsonWriterFactory<>::Create(&OutputString);
+	FJsonSerializer::Serialize(ToolsetEnvelope.ToSharedRef(), Writer);
+	return OutputString;
+}
diff --git a/Engine/Plugins/Experimental/ECABridge/Source/ECABridge/Private/ECAMCPServer.cpp b/Engine/Plugins/Experimental/ECABridge/Source/ECABridge/Private/ECAMCPServer.cpp
--- a/Engine/Plugins/Experimental/ECABridge/Source/ECABridge/Private/ECAMCPServer.cpp
+++ b/Engine/Plugins/Experimental/ECABridge/Source/ECABridge/Private/ECAMCPServer.cpp
@@ -11,6 +11,9 @@
 #include "Serialization/JsonReader.h"
 #include "Serialization/JsonWriter.h"
 #include "Serialization/JsonSerializer.h"
+#include "Editor.h"
+#include "ToolsetRegistry/ToolsetRegistry.h"
+#include "ToolsetRegistry/ToolsetRegistrySubsystem.h"
 
 FECAMCPServer::FECAMCPServer(UECABridge* InBridge)
 	: Bridge(InBridge)
@@ -378,44 +381,66 @@ TSharedPtr<FJsonObject> FECAMCPServer::HandleToolsCall(const TSharedPtr<FJsonObj
 
 	UE_LOG(LogTemp, Log, TEXT("[ECABridge] Tool call: %s"), *ToolName);
 
-	// Execute the command
-	FECACommandResult CommandResult = FECACommandRegistry::Get().ExecuteCommand(ToolName, Arguments);
-
-	// Build MCP response
-	TSharedPtr<FJsonObject> Result = MakeShared<FJsonObject>();
-	
-	TArray<TSharedPtr<FJsonValue>> Content;
-	TSharedPtr<FJsonObject> TextContent = MakeShared<FJsonObject>();
-	TextContent->SetStringField(TEXT("type"), TEXT("text"));
-	
-	if (CommandResult.bSuccess)
+	// First try ECABridge's own command registry
+	if (FECACommandRegistry::Get().HasCommand(ToolName))
 	{
-		FString ResultText;
-		if (CommandResult.ResultData.IsValid())
+		FECACommandResult CommandResult = FECACommandRegistry::Get().ExecuteCommand(ToolName, Arguments);
+
+		// Build MCP response
+		TSharedPtr<FJsonObject> Result = MakeShared<FJsonObject>();
+		
+		TArray<TSharedPtr<FJsonValue>> Content;
+		TSharedPtr<FJsonObject> TextContent = MakeShared<FJsonObject>();
+		TextContent->SetStringField(TEXT("type"), TEXT("text"));
+		
+		if (CommandResult.bSuccess)
 		{
-			ResultText = SerializeJson(CommandResult.ResultData);
+			FString ResultText;
+			if (CommandResult.ResultData.IsValid())
+			{
+				ResultText = SerializeJson(CommandResult.ResultData);
+			}
+			else
+			{
+				ResultText = TEXT("{\"success\": true}");
+			}
+			TextContent->SetStringField(TEXT("text"), ResultText);
 		}
 		else
 		{
-			ResultText = TEXT("{\"success\": true}");
+			Result->SetBoolField(TEXT("isError"), true);
+			TextContent->SetStringField(TEXT("text"), CommandResult.ErrorMessage);
 		}
-		TextContent->SetStringField(TEXT("text"), ResultText);
+		
+		Content.Add(MakeShared<FJsonValueObject>(TextContent));
+		Result->SetArrayField(TEXT("content"), Content);
+
+		return Result;
 	}
-	else
+	
+	// Direction 2: Try ToolsetRegistry if not found in ECABridge commands
+	TSharedPtr<FJsonObject> ToolsetResult;
+	if (TryExecuteToolsetRegistryTool(ToolName, Arguments, ToolsetResult))
 	{
-		Result->SetBoolField(TEXT("isError"), true);
-		TextContent->SetStringField(TEXT("text"), CommandResult.ErrorMessage);
+		return ToolsetResult;
 	}
 	
-	Content.Add(MakeShared<FJsonValueObject>(TextContent));
-	Result->SetArrayField(TEXT("content"), Content);
-
-	return Result;
+	// Tool not found anywhere
+	TSharedPtr<FJsonObject> ErrorResult = MakeShared<FJsonObject>();
+	ErrorResult->SetBoolField(TEXT("isError"), true);
+	TArray<TSharedPtr<FJsonValue>> ErrContent;
+	TSharedPtr<FJsonObject> ErrText = MakeShared<FJsonObject>();
+	ErrText->SetStringField(TEXT("type"), TEXT("text"));
+	ErrText->SetStringField(TEXT("text"), FString::Printf(TEXT("Unknown tool: %s"), *ToolName));
+	ErrContent.Add(MakeShared<FJsonValueObject>(ErrText));
+	ErrorResult->SetArrayField(TEXT("content"), ErrContent);
+	return ErrorResult;
 }
 
 TArray<TSharedPtr<FJsonValue>> FECAMCPServer::BuildToolDefinitions()
 {
 	TArray<TSharedPtr<FJsonValue>> Tools;
+	TSet<FString> ExistingNames;
 
 	for (const TSharedPtr<IECACommand>& Command : FECACommandRegistry::Get().GetAllCommands())
 	{
@@ -490,8 +515,13 @@ TArray<TSharedPtr<FJsonValue>> FECAMCPServer::BuildToolDefinitions()
 
 		ToolDef->SetObjectField(TEXT("inputSchema"), InputSchema);
 		Tools.Add(MakeShared<FJsonValueObject>(ToolDef));
+		ExistingNames.Add(Command->GetName().ToLower());
 	}
 
+	// Direction 2: Also include tools from ToolsetRegistry (excluding our own ECABridge toolset
+	// to avoid duplication — we already added those commands above)
+	AppendToolsetRegistryTools(Tools, ExistingNames);
+
 	return Tools;
 }
 
@@ -532,6 +562,188 @@ TSharedPtr<FJsonObject> FECAMCPServer::CreateJsonRpcError(int32 Code, const FStr
 	return Response;
 }
 
+void FECAMCPServer::AppendToolsetRegistryTools(TArray<TSharedPtr<FJsonValue>>& Tools, TSet<FString>& ExistingNames)
+{
+#if WITH_EDITOR
+	UToolsetRegistrySubsystem* ToolsetRegistrySubsystem = 
+		GEditor ? GEditor->GetEditorSubsystem<UToolsetRegistrySubsystem>() : nullptr;
+	if (!ToolsetRegistrySubsystem)
+	{
+		return;
+	}
+	
+	UE::ToolsetRegistry::FToolsetRegistry& ToolsetRegistry = ToolsetRegistrySubsystem->ToolsetRegistry;
+	const FString SchemasJson = ToolsetRegistry.GetToolsetJsonSchemas();
+	
+	TArray<TSharedPtr<FJsonValue>> ToolsetArray;
+	TSharedRef<TJsonReader<>> Reader = TJsonReaderFactory<>::Create(SchemasJson);
+	if (!FJsonSerializer::Deserialize(Reader, ToolsetArray))
+	{
+		UE_LOG(LogTemp, Warning, TEXT("[ECABridge] Failed to parse ToolsetRegistry schemas"));
+		return;
+	}
+	
+	int32 AddedCount = 0;
+	
+	for (const TSharedPtr<FJsonValue>& ToolsetValue : ToolsetArray)
+	{
+		const TSharedPtr<FJsonObject>& ToolsetObject = ToolsetValue->AsObject();
+		if (!ToolsetObject.IsValid())
+		{
+			continue;
+		}
+		
+		// Skip our own ECABridge toolset to avoid duplication
+		FString ToolsetName;
+		ToolsetObject->TryGetStringField(TEXT("name"), ToolsetName);
+		if (ToolsetName == TEXT("ECABridge"))
+		{
+			continue;
+		}
+		
+		const TArray<TSharedPtr<FJsonValue>>* ToolsField = nullptr;
+		if (!ToolsetObject->TryGetArrayField(TEXT("tools"), ToolsField) || !ToolsField)
+		{
+			continue;
+		}
+		
+		for (const TSharedPtr<FJsonValue>& ToolValue : *ToolsField)
+		{
+			const TSharedPtr<FJsonObject>& ToolJsonObject = ToolValue->AsObject();
+			if (!ToolJsonObject.IsValid())
+			{
+				continue;
+			}
+			
+			FString ToolName;
+			ToolJsonObject->TryGetStringField(TEXT("name"), ToolName);
+			if (ToolName.IsEmpty())
+			{
+				continue;
+			}
+			
+			// Skip duplicates
+			if (ExistingNames.Contains(ToolName.ToLower()))
+			{
+				UE_LOG(LogTemp, Verbose, TEXT("[ECABridge] Skipping duplicate ToolsetRegistry tool: %s"), *ToolName);
+				continue;
+			}
+			
+			// Build MCP tool definition from ToolsetRegistry schema
+			TSharedPtr<FJsonObject> McpToolDef = MakeShared<FJsonObject>();
+			McpToolDef->SetStringField(TEXT("name"), ToolName);
+			
+			FString ToolDescription;
+			if (ToolJsonObject->TryGetStringField(TEXT("description"), ToolDescription) && !ToolDescription.IsEmpty())
+			{
+				McpToolDef->SetStringField(TEXT("description"), ToolDescription);
+			}
+			
+			const TSharedPtr<FJsonObject>* InputSchemaField = nullptr;
+			if (ToolJsonObject->TryGetObjectField(TEXT("inputSchema"), InputSchemaField) && InputSchemaField)
+			{
+				McpToolDef->SetObjectField(TEXT("inputSchema"), *InputSchemaField);
+			}
+			
+			Tools.Add(MakeShared<FJsonValueObject>(McpToolDef));
+			ExistingNames.Add(ToolName.ToLower());
+			AddedCount++;
+		}
+	}
+	
+	if (AddedCount > 0)
+	{
+		UE_LOG(LogTemp, Log, TEXT("[ECABridge] Added %d tools from ToolsetRegistry to MCP tool list"), AddedCount);
+	}
+#endif // WITH_EDITOR
+}
+
+bool FECAMCPServer::TryExecuteToolsetRegistryTool(const FString& ToolName, const TSharedPtr<FJsonObject>& Arguments,
+	TSharedPtr<FJsonObject>& OutResult)
+{
+#if WITH_EDITOR
+	UToolsetRegistrySubsystem* ToolsetRegistrySubsystem = 
+		GEditor ? GEditor->GetEditorSubsystem<UToolsetRegistrySubsystem>() : nullptr;
+	if (!ToolsetRegistrySubsystem)
+	{
+		return false;
+	}
+	
+	UE::ToolsetRegistry::FToolsetRegistry& ToolsetRegistry = ToolsetRegistrySubsystem->ToolsetRegistry;
+	
+	// Serialize arguments to JSON string for the ToolsetRegistry API
+	FString ArgumentsJson;
+	if (Arguments.IsValid())
+	{
+		TSharedRef<TJsonWriter<>> Writer = TJsonWriterFactory<>::Create(&ArgumentsJson);
+		FJsonSerializer::Serialize(Arguments.ToSharedRef(), Writer);
+	}
+	else
+	{
+		ArgumentsJson = TEXT("{}");
+	}
+	
+	UE_LOG(LogTemp, Log, TEXT("[ECABridge] Executing ToolsetRegistry tool: '%s'"), *ToolName);
+	
+	// ExecuteTool returns a TFuture. Since ECABridge's HTTP handler runs on the game thread
+	// (FHttpServerModule ticks on game thread), and ToolsetRegistry's internal RunOnMainThread
+	// checks IsInGameThread() and executes inline, the future is typically fulfilled immediately
+	// for synchronous tools. For async tools, we wait with a timeout.
+	TFuture<TValueOrError<FString, FString>> Future = ToolsetRegistry.ExecuteTool(ToolName, ArgumentsJson);
+	
+	// Check if the future is already ready (synchronous tools)
+	if (!Future.IsReady())
+	{
+		// For async tools, we need to pump the game thread task graph to allow the future to complete
+		// since we're blocking the game thread here.
+		double StartTime = FPlatformTime::Seconds();
+		constexpr double TimeoutSeconds = 30.0;
+		while (!Future.IsReady() && (FPlatformTime::Seconds() - StartTime) < TimeoutSeconds)
+		{
+			FTaskGraphInterface::Get().ProcessThreadUntilIdle(ENamedThreads::GameThread);
+			FPlatformProcess::Sleep(0.01f);
+		}
+	}
+	
+	if (!Future.IsReady())
+	{
+		OutResult = MakeShared<FJsonObject>();
+		OutResult->SetBoolField(TEXT("isError"), true);
+		TArray<TSharedPtr<FJsonValue>> Content;
+		TSharedPtr<FJsonObject> TextContent = MakeShared<FJsonObject>();
+		TextContent->SetStringField(TEXT("type"), TEXT("text"));
+		TextContent->SetStringField(TEXT("text"), 
+			FString::Printf(TEXT("ToolsetRegistry tool '%s' timed out after 30 seconds"), *ToolName));
+		Content.Add(MakeShared<FJsonValueObject>(TextContent));
+		OutResult->SetArrayField(TEXT("content"), Content);
+		return true;
+	}
+	
+	TValueOrError<FString, FString> ToolResult = Future.Get();
+	
+	OutResult = MakeShared<FJsonObject>();
+	TArray<TSharedPtr<FJsonValue>> Content;
+	TSharedPtr<FJsonObject> TextContent = MakeShared<FJsonObject>();
+	TextContent->SetStringField(TEXT("type"), TEXT("text"));
+	
+	if (ToolResult.HasError())
+	{
+		OutResult->SetBoolField(TEXT("isError"), true);
+		TextContent->SetStringField(TEXT("text"), ToolResult.GetError());
+	}
+	else
+	{
+		TextContent->SetStringField(TEXT("text"), ToolResult.GetValue());
+	}
+	
+	Content.Add(MakeShared<FJsonValueObject>(TextContent));
+	OutResult->SetArrayField(TEXT("content"), Content);
+	return true;
+#else
+	return false;
+#endif // WITH_EDITOR
+}
+
 FString FECAMCPServer::SerializeJson(const TSharedPtr<FJsonObject>& JsonObject)
 {
 	FString OutputString;
diff --git a/Engine/Plugins/Experimental/ECABridge/Source/ECABridge/Public/ECABridgeModule.h b/Engine/Plugins/Experimental/ECABridge/Source/ECABridge/Public/ECABridgeModule.h
--- a/Engine/Plugins/Experimental/ECABridge/Source/ECABridge/Public/ECABridgeModule.h
+++ b/Engine/Plugins/Experimental/ECABridge/Source/ECABridge/Public/ECABridgeModule.h
@@ -6,11 +6,14 @@
 #include "Modules/ModuleManager.h"
 #include "Engine/EngineTypes.h"  // For FTimerHandle
 #include "Editor/UnrealEdTypes.h"  // For ELevelViewportType
+#include "ToolsetRegistry/Toolset.h"
+#include "ToolsetRegistry/ToolsetRegistry.h"
 
 class UECABridge;
 class UFactory;
 struct FAssetData;
 class AActor;
+class FECABridgeToolset;
 
 /**
  * ECA Bridge Module
@@ -50,6 +53,12 @@ private:
 	/** Console commands for debugging */
 	void RegisterConsoleCommands();
 	
+	/** Register ECABridge commands as a ToolsetRegistry toolset */
+	void RegisterWithToolsetRegistry();
+	
+	/** Unregister from ToolsetRegistry */
+	void UnregisterFromToolsetRegistry();
+	
 	/** Register editor event handlers for pending events */
 	void RegisterEditorEventHandlers();
 	
@@ -140,4 +149,14 @@ private:
 	
 	/** Debounce time in seconds - wait this long after camera stops before sending event */
 	static constexpr float CameraStoppedDebounceTime = 0.5f;
+	
+	//--------------------------------------------------------------------------
+	// ToolsetRegistry Integration
+	//--------------------------------------------------------------------------
+	
+	/** The ECABridge toolset registered with ToolsetRegistry (Direction 1) */
+	TSharedPtr<FECABridgeToolset> ECABridgeToolset;
+	
+	/** Handle for ToolsetRegistry change notifications */
+	FDelegateHandle ToolsetRegistryChangedHandle;
 };
diff --git a/Engine/Plugins/Experimental/ECABridge/Source/ECABridge/Public/ECABridgeToolset.h b/Engine/Plugins/Experimental/ECABridge/Source/ECABridge/Public/ECABridgeToolset.h
new file mode 100644
--- /dev/null
+++ b/Engine/Plugins/Experimental/ECABridge/Source/ECABridge/Public/ECABridgeToolset.h
@@ -0,0 +1,33 @@
+// Copyright Epic Games, Inc. All Rights Reserved.
+
+#pragma once
+
+#include "CoreMinimal.h"
+#include "ToolsetRegistry/Toolset.h"
+
+/**
+ * ECABridge Toolset
+ * 
+ * Exposes all ECABridge commands (229+) to the ToolsetRegistry so that other 
+ * consumers (AIAssistant, Claude Code, MCP plugin, Monologue, etc.) can discover
+ * and invoke them through the standard UE toolset interface.
+ * 
+ * This wraps FECACommandRegistry, translating between IECACommand and FToolset interfaces:
+ *   - GetToolsetJsonSchema() builds JSON Schema from IECACommand::GetParameters()
+ *   - ExecuteTool() forwards to FECACommandRegistry::ExecuteCommand()
+ */
+class ECABRIDGE_API FECABridgeToolset : public UE::ToolsetRegistry::FToolset
+{
+public:
+	FECABridgeToolset();
+	virtual ~FECABridgeToolset();
+
+	//~ Begin FToolset interface
+	virtual TFuture<TValueOrError<FString, FString>> ExecuteTool(
+		const FString& ToolName, const FString& JsonInput) override;
+
+	virtual FString GetToolsetJsonSchema() const override;
+	virtual FString GetToolsetName() const override;
+	virtual FString GetToolsetVersion() const override;
+	//~ End FToolset interface
+};
diff --git a/Engine/Plugins/Experimental/ECABridge/Source/ECABridge/Public/ECAMCPServer.h b/Engine/Plugins/Experimental/ECABridge/Source/ECABridge/Public/ECAMCPServer.h
--- a/Engine/Plugins/Experimental/ECABridge/Source/ECABridge/Public/ECAMCPServer.h
+++ b/Engine/Plugins/Experimental/ECABridge/Source/ECABridge/Public/ECAMCPServer.h
@@ -64,9 +64,16 @@ private:
 	/** Handle tools/call request */
 	TSharedPtr<FJsonObject> HandleToolsCall(const TSharedPtr<FJsonObject>& Params);
 	
-	/** Build MCP tool definitions from registered commands */
+	/** Build MCP tool definitions from registered commands and ToolsetRegistry */
 	TArray<TSharedPtr<FJsonValue>> BuildToolDefinitions();
 	
+	/** Build tool definitions from ToolsetRegistry toolsets (Direction 2) */
+	void AppendToolsetRegistryTools(TArray<TSharedPtr<FJsonValue>>& Tools, TSet<FString>& ExistingNames);
+	
+	/** Try to execute a tool via ToolsetRegistry if not found in ECABridge commands */
+	bool TryExecuteToolsetRegistryTool(const FString& ToolName, const TSharedPtr<FJsonObject>& Arguments,
+		TSharedPtr<FJsonObject>& OutResult);
+	
 	/** Create JSON-RPC response wrapper */
 	TSharedPtr<FJsonObject> CreateJsonRpcResponse(const TSharedPtr<FJsonObject>& Result, const TSharedPtr<FJsonValue>& RequestId);
 	
