Added support to read Behavior Trees, Blackboards, State Trees and EQS assets in ECA bridge.

#rnx
#Synced-CL 51885519

diff --git a/Engine/Plugins/Experimental/ECABridge/Source/ECABridge/ECABridge.Build.cs b/Engine/Plugins/Experimental/ECABridge/Source/ECABridge/ECABridge.Build.cs
--- a/Engine/Plugins/Experimental/ECABridge/Source/ECABridge/ECABridge.Build.cs
+++ b/Engine/Plugins/Experimental/ECABridge/Source/ECABridge/ECABridge.Build.cs
@@ -56,6 +56,11 @@ public class ECABridge : ModuleRules
 			"MaterialEditor",            // For material editing utilities
 			"DataTableEditor",           // For DataTable editor UI refresh (BroadcastPostChange)
 			
+			// UE AI modules
+			"AIModule",                  // For Behavior Tree and Blackboard asset access
+			"StateTreeModule",           // For UStateTree and StateTree types
+			"StateTreeEditorModule",     // For UStateTreeEditorData, UStateTreeState, editor node types
+			
 			// Additional utilities
 			"Projects", "DeveloperSettings", "ToolMenus",
 			
diff --git a/Engine/Plugins/Experimental/ECABridge/Source/ECABridge/Private/Commands/ECABehaviorTreeCommands.cpp b/Engine/Plugins/Experimental/ECABridge/Source/ECABridge/Private/Commands/ECABehaviorTreeCommands.cpp
new file mode 100644
--- /dev/null
+++ b/Engine/Plugins/Experimental/ECABridge/Source/ECABridge/Private/Commands/ECABehaviorTreeCommands.cpp
@@ -0,0 +1,284 @@
+// Copyright Epic Games, Inc. All Rights Reserved.
+
+#include "Commands/ECABehaviorTreeCommands.h"
+#include "Commands/ECACommand.h"
+#include "BehaviorTree/BehaviorTree.h"
+#include "BehaviorTree/BTCompositeNode.h"
+#include "BehaviorTree/BTTaskNode.h"
+#include "BehaviorTree/BTDecorator.h"
+#include "BehaviorTree/BTService.h"
+#include "BehaviorTree/BlackboardData.h"
+#include "BehaviorTree/Blackboard/BlackboardKeyType.h"
+#include "BehaviorTree/Composites/BTComposite_Selector.h"
+#include "BehaviorTree/Composites/BTComposite_Sequence.h"
+#include "BehaviorTree/Composites/BTComposite_SimpleParallel.h"
+
+// Register the commands
+REGISTER_ECA_COMMAND(FECACommand_GetBehaviorTreeInfo)
+REGISTER_ECA_COMMAND(FECACommand_GetBlackboardInfo)
+
+//------------------------------------------------------------------------------
+// SerializeBlackboardData — serialize blackboard keys including parent chain
+// (file-scoped helper shared by GetBehaviorTreeInfo and GetBlackboardInfo)
+//------------------------------------------------------------------------------
+
+static TSharedPtr<FJsonObject> SerializeBlackboardData(UBlackboardData* BBData)
+{
+	if (!BBData)
+	{
+		return MakeShared<FJsonObject>();
+	}
+
+	TSharedPtr<FJsonObject> Obj = MakeShared<FJsonObject>();
+	Obj->SetStringField(TEXT("asset_path"), BBData->GetPathName());
+	Obj->SetStringField(TEXT("asset_name"), BBData->GetName());
+
+	// Walk parent chain so we include inherited keys (parents first)
+	TArray<UBlackboardData*> DataChain;
+	UBlackboardData* Current = BBData;
+	while (Current)
+	{
+		DataChain.Insert(Current, 0);
+		Current = Current->Parent;
+	}
+
+	if (BBData->Parent)
+	{
+		Obj->SetStringField(TEXT("parent_asset"), BBData->Parent->GetPathName());
+	}
+
+	TArray<TSharedPtr<FJsonValue>> KeysArray;
+	for (UBlackboardData* Data : DataChain)
+	{
+		if (!Data)
+		{
+			continue;
+		}
+		
+		bool bIsInherited = (Data != BBData);
+		for (const FBlackboardEntry& Entry : Data->Keys)
+		{
+			TSharedPtr<FJsonObject> KeyObj = MakeShared<FJsonObject>();
+			KeyObj->SetStringField(TEXT("name"), Entry.EntryName.ToString());
+			KeyObj->SetStringField(TEXT("type"), Entry.KeyType ? Entry.KeyType->GetClass()->GetName() : TEXT("None"));
+			KeyObj->SetBoolField(TEXT("instance_synced"), Entry.bInstanceSynced != 0);
+			if (bIsInherited)
+			{
+				KeyObj->SetBoolField(TEXT("inherited"), true);
+			}
+#if WITH_EDITORONLY_DATA
+			if (!Entry.EntryDescription.IsEmpty())
+			{
+				KeyObj->SetStringField(TEXT("description"), Entry.EntryDescription);
+			}
+			if (!Entry.EntryCategory.IsNone())
+			{
+				KeyObj->SetStringField(TEXT("category"), Entry.EntryCategory.ToString());
+			}
+#endif
+			KeysArray.Add(MakeShared<FJsonValueObject>(KeyObj));
+		}
+	}
+
+	Obj->SetArrayField(TEXT("keys"), KeysArray);
+	Obj->SetNumberField(TEXT("key_count"), KeysArray.Num());
+
+	return Obj;
+}
+
+//------------------------------------------------------------------------------
+// GetBehaviorTreeInfo
+//------------------------------------------------------------------------------
+
+FECACommandResult FECACommand_GetBehaviorTreeInfo::Execute(const TSharedPtr<FJsonObject>& Params)
+{
+	FString AssetPath;
+	if (!GetStringParam(Params, TEXT("asset_path"), AssetPath))
+	{
+		return FECACommandResult::Error(TEXT("Missing required parameter: asset_path"));
+	}
+
+	int32 MaxDepth = -1;
+	GetIntParam(Params, TEXT("max_depth"), MaxDepth, /*bRequired=*/false);
+
+	// Load the behavior tree asset
+	UBehaviorTree* BT = LoadObject<UBehaviorTree>(nullptr, *AssetPath);
+	if (!BT)
+	{
+		return FECACommandResult::Error(FString::Printf(TEXT("Behavior Tree not found: %s"), *AssetPath));
+	}
+
+	TSharedPtr<FJsonObject> Result = MakeResult();
+	Result->SetStringField(TEXT("asset_path"), AssetPath);
+	Result->SetStringField(TEXT("asset_name"), BT->GetName());
+	Result->SetStringField(TEXT("asset_class"), TEXT("BehaviorTree"));
+
+	// Serialize blackboard
+	if (BT->BlackboardAsset)
+	{
+		Result->SetObjectField(TEXT("blackboard"), SerializeBlackboardData(BT->BlackboardAsset));
+	}
+
+	// Serialize root decorators (if any — used by subtrees)
+	if (BT->RootDecorators.Num() > 0)
+	{
+		TArray<TSharedPtr<FJsonValue>> RootDecArray;
+		for (UBTDecorator* Dec : BT->RootDecorators)
+		{
+			if (Dec)
+			{
+				RootDecArray.Add(MakeShared<FJsonValueObject>(SerializeNode(Dec)));
+			}
+		}
+		Result->SetArrayField(TEXT("root_decorators"), RootDecArray);
+	}
+
+	// Serialize the node tree
+	if (BT->RootNode)
+	{
+		Result->SetObjectField(TEXT("root"), SerializeCompositeNode(BT->RootNode, MaxDepth, 0));
+	}
+
+	return FECACommandResult::Success(Result);
+}
+
+//------------------------------------------------------------------------------
+// SerializeNode — base serialization for any UBTNode
+//------------------------------------------------------------------------------
+
+TSharedPtr<FJsonObject> FECACommand_GetBehaviorTreeInfo::SerializeNode(UBTNode* Node) const
+{
+	if (!Node)
+	{
+		return MakeShared<FJsonObject>();
+	}
+
+	TSharedPtr<FJsonObject> Obj = MakeShared<FJsonObject>();
+
+	Obj->SetStringField(TEXT("node_name"), Node->NodeName.IsEmpty() ? Node->GetClass()->GetName() : Node->NodeName);
+	Obj->SetStringField(TEXT("node_class"), Node->GetClass()->GetName());
+	Obj->SetStringField(TEXT("description"), Node->GetStaticDescription());
+	Obj->SetNumberField(TEXT("execution_index"), Node->GetExecutionIndex());
+	Obj->SetNumberField(TEXT("depth"), Node->GetTreeDepth());
+
+	return Obj;
+}
+
+//------------------------------------------------------------------------------
+// SerializeCompositeNode — recursive serialization of composite + children
+//------------------------------------------------------------------------------
+
+TSharedPtr<FJsonObject> FECACommand_GetBehaviorTreeInfo::SerializeCompositeNode(
+	UBTCompositeNode* CompositeNode, int32 MaxDepth, int32 CurrentDepth) const
+{
+	if (!CompositeNode)
+	{
+		return MakeShared<FJsonObject>();
+	}
+
+	TSharedPtr<FJsonObject> Obj = SerializeNode(CompositeNode);
+
+	// Determine composite type from class
+	FString CompositeType = TEXT("composite");
+	if (CompositeNode->IsA<UBTComposite_Selector>())
+	{
+		CompositeType = TEXT("selector");
+	}
+	else if (CompositeNode->IsA<UBTComposite_Sequence>())
+	{
+		CompositeType = TEXT("sequence");
+	}
+	else if (CompositeNode->IsA<UBTComposite_SimpleParallel>())
+	{
+		CompositeType = TEXT("parallel");
+	}
+	Obj->SetStringField(TEXT("composite_type"), CompositeType);
+
+	// Services attached to this composite
+	if (CompositeNode->Services.Num() > 0)
+	{
+		TArray<TSharedPtr<FJsonValue>> ServicesArray;
+		for (UBTService* Service : CompositeNode->Services)
+		{
+			if (Service)
+			{
+				ServicesArray.Add(MakeShared<FJsonValueObject>(SerializeNode(Service)));
+			}
+		}
+		Obj->SetArrayField(TEXT("services"), ServicesArray);
+	}
+
+	// Children
+	bool bShouldRecurse = (MaxDepth < 0 || CurrentDepth < MaxDepth);
+	TArray<TSharedPtr<FJsonValue>> ChildrenArray;
+
+	for (const FBTCompositeChild& Child : CompositeNode->Children)
+	{
+		TSharedPtr<FJsonObject> ChildObj = MakeShared<FJsonObject>();
+
+		// Decorators on this child branch
+		if (Child.Decorators.Num() > 0)
+		{
+			TArray<TSharedPtr<FJsonValue>> DecArray;
+			for (UBTDecorator* Dec : Child.Decorators)
+			{
+				if (Dec)
+				{
+					DecArray.Add(MakeShared<FJsonValueObject>(SerializeNode(Dec)));
+				}
+			}
+			ChildObj->SetArrayField(TEXT("decorators"), DecArray);
+		}
+
+		// The actual child node — either a composite (branch continues) or a task (leaf)
+		if (Child.ChildComposite && bShouldRecurse)
+		{
+			ChildObj->SetObjectField(TEXT("node"),
+				SerializeCompositeNode(Child.ChildComposite, MaxDepth, CurrentDepth + 1));
+		}
+		else if (Child.ChildComposite)
+		{
+			// Hit depth limit — serialize node info without recursing into children
+			TSharedPtr<FJsonObject> ShallowNode = SerializeNode(Child.ChildComposite);
+			ShallowNode->SetBoolField(TEXT("truncated"), true);
+			ShallowNode->SetNumberField(TEXT("child_count"), Child.ChildComposite->GetChildrenNum());
+			ChildObj->SetObjectField(TEXT("node"), ShallowNode);
+		}
+		else if (Child.ChildTask)
+		{
+			ChildObj->SetObjectField(TEXT("node"), SerializeNode(Child.ChildTask));
+		}
+
+		ChildrenArray.Add(MakeShared<FJsonValueObject>(ChildObj));
+	}
+
+	if (ChildrenArray.Num() > 0)
+	{
+		Obj->SetArrayField(TEXT("children"), ChildrenArray);
+	}
+
+	return Obj;
+}
+
+
+//------------------------------------------------------------------------------
+// GetBlackboardInfo
+//------------------------------------------------------------------------------
+
+FECACommandResult FECACommand_GetBlackboardInfo::Execute(const TSharedPtr<FJsonObject>& Params)
+{
+	FString AssetPath;
+	if (!GetStringParam(Params, TEXT("asset_path"), AssetPath))
+	{
+		return FECACommandResult::Error(TEXT("Missing required parameter: asset_path"));
+	}
+
+	UBlackboardData* BBData = LoadObject<UBlackboardData>(nullptr, *AssetPath);
+	if (!BBData)
+	{
+		return FECACommandResult::Error(FString::Printf(TEXT("Blackboard not found: %s"), *AssetPath));
+	}
+
+	TSharedPtr<FJsonObject> Result = SerializeBlackboardData(BBData);
+	return FECACommandResult::Success(Result);
+}
diff --git a/Engine/Plugins/Experimental/ECABridge/Source/ECABridge/Private/Commands/ECAEQSCommands.cpp b/Engine/Plugins/Experimental/ECABridge/Source/ECABridge/Private/Commands/ECAEQSCommands.cpp
new file mode 100644
--- /dev/null
+++ b/Engine/Plugins/Experimental/ECABridge/Source/ECABridge/Private/Commands/ECAEQSCommands.cpp
@@ -0,0 +1,265 @@
+// Copyright Epic Games, Inc. All Rights Reserved.
+
+#include "Commands/ECAEQSCommands.h"
+#include "Commands/ECACommand.h"
+#include "EnvironmentQuery/EnvQuery.h"
+#include "EnvironmentQuery/EnvQueryOption.h"
+#include "EnvironmentQuery/EnvQueryGenerator.h"
+#include "EnvironmentQuery/EnvQueryTest.h"
+
+// Register the command
+REGISTER_ECA_COMMAND(FECACommand_GetEQSInfo)
+
+//------------------------------------------------------------------------------
+// Helper: Convert EQS enums to strings
+//------------------------------------------------------------------------------
+
+static FString TestPurposeToString(EEnvTestPurpose::Type Purpose)
+{
+	switch (Purpose)
+	{
+	case EEnvTestPurpose::Filter:         return TEXT("Filter");
+	case EEnvTestPurpose::Score:           return TEXT("Score");
+	case EEnvTestPurpose::FilterAndScore:  return TEXT("FilterAndScore");
+	default:                               return TEXT("Unknown");
+	}
+}
+
+static FString FilterTypeToString(EEnvTestFilterType::Type FilterType)
+{
+	switch (FilterType)
+	{
+	case EEnvTestFilterType::Minimum:  return TEXT("Minimum");
+	case EEnvTestFilterType::Maximum:  return TEXT("Maximum");
+	case EEnvTestFilterType::Range:    return TEXT("Range");
+	case EEnvTestFilterType::Match:    return TEXT("Match");
+	default:                           return TEXT("Unknown");
+	}
+}
+
+static FString ScoreEquationToString(EEnvTestScoreEquation::Type Equation)
+{
+	switch (Equation)
+	{
+	case EEnvTestScoreEquation::Linear:        return TEXT("Linear");
+	case EEnvTestScoreEquation::Square:         return TEXT("Square");
+	case EEnvTestScoreEquation::InverseLinear:  return TEXT("InverseLinear");
+	case EEnvTestScoreEquation::SquareRoot:     return TEXT("SquareRoot");
+	case EEnvTestScoreEquation::Constant:       return TEXT("Constant");
+	default:                                    return TEXT("Unknown");
+	}
+}
+
+static FString TestCostToString(EEnvTestCost::Type Cost)
+{
+	switch (Cost)
+	{
+	case EEnvTestCost::Low:     return TEXT("Low");
+	case EEnvTestCost::Medium:  return TEXT("Medium");
+	case EEnvTestCost::High:    return TEXT("High");
+	default:                    return TEXT("Unknown");
+	}
+}
+
+//------------------------------------------------------------------------------
+// GetEQSInfo
+//------------------------------------------------------------------------------
+
+FECACommandResult FECACommand_GetEQSInfo::Execute(const TSharedPtr<FJsonObject>& Params)
+{
+	FString AssetPath;
+	if (!GetStringParam(Params, TEXT("asset_path"), AssetPath))
+	{
+		return FECACommandResult::Error(TEXT("Missing required parameter: asset_path"));
+	}
+
+	UEnvQuery* Query = LoadObject<UEnvQuery>(nullptr, *AssetPath);
+	if (!Query)
+	{
+		return FECACommandResult::Error(FString::Printf(TEXT("EQS query not found: %s"), *AssetPath));
+	}
+
+	TSharedPtr<FJsonObject> Result = MakeResult();
+	Result->SetStringField(TEXT("asset_path"), AssetPath);
+	Result->SetStringField(TEXT("asset_name"), Query->GetName());
+	Result->SetStringField(TEXT("asset_class"), TEXT("EnvQuery"));
+	Result->SetStringField(TEXT("query_name"), Query->GetQueryName().ToString());
+
+	// Options (each is a generator + tests)
+	const TArray<UEnvQueryOption*>& Options = Query->GetOptions();
+	if (Options.Num() > 0)
+	{
+		TArray<TSharedPtr<FJsonValue>> OptionsArray;
+		for (const UEnvQueryOption* Option : Options)
+		{
+			if (Option)
+			{
+				OptionsArray.Add(MakeShared<FJsonValueObject>(SerializeOption(Option)));
+			}
+		}
+		Result->SetArrayField(TEXT("options"), OptionsArray);
+	}
+
+	return FECACommandResult::Success(Result);
+}
+
+//------------------------------------------------------------------------------
+// SerializeOption — a generator paired with its tests
+//------------------------------------------------------------------------------
+
+TSharedPtr<FJsonObject> FECACommand_GetEQSInfo::SerializeOption(const UEnvQueryOption* Option) const
+{
+	TSharedPtr<FJsonObject> Obj = MakeShared<FJsonObject>();
+
+	if (!Option)
+	{
+		return Obj;
+	}
+
+	// Description
+	FText Title = Option->GetDescriptionTitle();
+	if (!Title.IsEmpty())
+	{
+		Obj->SetStringField(TEXT("description_title"), Title.ToString());
+	}
+
+	// Generator
+	if (Option->Generator)
+	{
+		Obj->SetObjectField(TEXT("generator"), SerializeGenerator(Option->Generator));
+	}
+
+	// Tests
+	if (Option->Tests.Num() > 0)
+	{
+		TArray<TSharedPtr<FJsonValue>> TestsArray;
+		for (const UEnvQueryTest* Test : Option->Tests)
+		{
+			if (Test)
+			{
+				TestsArray.Add(MakeShared<FJsonValueObject>(SerializeTest(Test)));
+			}
+		}
+		Obj->SetArrayField(TEXT("tests"), TestsArray);
+	}
+
+	return Obj;
+}
+
+//------------------------------------------------------------------------------
+// SerializeGenerator — the item generator for an option
+//------------------------------------------------------------------------------
+
+TSharedPtr<FJsonObject> FECACommand_GetEQSInfo::SerializeGenerator(const UEnvQueryGenerator* Generator) const
+{
+	TSharedPtr<FJsonObject> Obj = MakeShared<FJsonObject>();
+
+	if (!Generator)
+	{
+		return Obj;
+	}
+
+	Obj->SetStringField(TEXT("class"), Generator->GetClass()->GetName());
+
+	// Option name (user-facing label)
+	if (!Generator->OptionName.IsEmpty())
+	{
+		Obj->SetStringField(TEXT("option_name"), Generator->OptionName);
+	}
+
+	// Item type
+	if (Generator->ItemType)
+	{
+		Obj->SetStringField(TEXT("item_type"), Generator->ItemType->GetName());
+	}
+
+	// Human-readable descriptions from the node
+	FText Title = Generator->GetDescriptionTitle();
+	if (!Title.IsEmpty())
+	{
+		Obj->SetStringField(TEXT("description_title"), Title.ToString());
+	}
+
+	FText Details = Generator->GetDescriptionDetails();
+	if (!Details.IsEmpty())
+	{
+		Obj->SetStringField(TEXT("description_details"), Details.ToString());
+	}
+
+	return Obj;
+}
+
+//------------------------------------------------------------------------------
+// SerializeTest — a filter/scoring test
+//------------------------------------------------------------------------------
+
+TSharedPtr<FJsonObject> FECACommand_GetEQSInfo::SerializeTest(const UEnvQueryTest* Test) const
+{
+	TSharedPtr<FJsonObject> Obj = MakeShared<FJsonObject>();
+
+	if (!Test)
+	{
+		return Obj;
+	}
+
+	Obj->SetStringField(TEXT("class"), Test->GetClass()->GetName());
+
+	// Human-readable descriptions
+	FText Title = Test->GetDescriptionTitle();
+	if (!Title.IsEmpty())
+	{
+		Obj->SetStringField(TEXT("description_title"), Title.ToString());
+	}
+
+	FText Details = Test->GetDescriptionDetails();
+	if (!Details.IsEmpty())
+	{
+		Obj->SetStringField(TEXT("description_details"), Details.ToString());
+	}
+
+	// Test purpose
+	Obj->SetStringField(TEXT("test_purpose"), TestPurposeToString(Test->TestPurpose));
+
+	// Test comment
+	if (!Test->TestComment.IsEmpty())
+	{
+		Obj->SetStringField(TEXT("test_comment"), Test->TestComment);
+	}
+
+	// Test order
+	Obj->SetNumberField(TEXT("test_order"), Test->TestOrder);
+
+	// Cost
+	Obj->SetStringField(TEXT("cost"), TestCostToString(Test->Cost));
+
+	// Filter settings (only relevant when filtering)
+	if (Test->IsFiltering())
+	{
+		Obj->SetStringField(TEXT("filter_type"), FilterTypeToString(Test->FilterType));
+
+		if (Test->GetWorkOnFloatValues())
+		{
+			if (Test->FilterType == EEnvTestFilterType::Minimum || Test->FilterType == EEnvTestFilterType::Range)
+			{
+				Obj->SetNumberField(TEXT("filter_min"), Test->FloatValueMin.DefaultValue);
+			}
+			if (Test->FilterType == EEnvTestFilterType::Maximum || Test->FilterType == EEnvTestFilterType::Range)
+			{
+				Obj->SetNumberField(TEXT("filter_max"), Test->FloatValueMax.DefaultValue);
+			}
+		}
+		else
+		{
+			Obj->SetBoolField(TEXT("filter_match_value"), Test->BoolValue.DefaultValue);
+		}
+	}
+
+	// Scoring settings (only relevant when scoring)
+	if (Test->IsScoring())
+	{
+		Obj->SetStringField(TEXT("scoring_equation"), ScoreEquationToString(Test->ScoringEquation));
+		Obj->SetNumberField(TEXT("scoring_factor"), Test->ScoringFactor.DefaultValue);
+	}
+
+	return Obj;
+}
diff --git a/Engine/Plugins/Experimental/ECABridge/Source/ECABridge/Private/Commands/ECAStateTreeCommands.cpp b/Engine/Plugins/Experimental/ECABridge/Source/ECABridge/Private/Commands/ECAStateTreeCommands.cpp
new file mode 100644
--- /dev/null
+++ b/Engine/Plugins/Experimental/ECABridge/Source/ECABridge/Private/Commands/ECAStateTreeCommands.cpp
@@ -0,0 +1,479 @@
+// Copyright Epic Games, Inc. All Rights Reserved.
+
+#include "Commands/ECAStateTreeCommands.h"
+#include "Commands/ECACommand.h"
+#include "StateTree.h"
+#include "StateTreeEditorData.h"
+#include "StateTreeState.h"
+#include "StateTreeEditorNode.h"
+#include "StateTreeSchema.h"
+#include "StateTreeNodeBase.h"
+
+// Register the command
+REGISTER_ECA_COMMAND(FECACommand_GetStateTreeInfo)
+
+//------------------------------------------------------------------------------
+// Helper: Convert enum to string
+//------------------------------------------------------------------------------
+
+static FString StateTypeToString(EStateTreeStateType Type)
+{
+	switch (Type)
+	{
+	case EStateTreeStateType::State:       return TEXT("State");
+	case EStateTreeStateType::Group:       return TEXT("Group");
+	case EStateTreeStateType::Linked:      return TEXT("Linked");
+	case EStateTreeStateType::LinkedAsset: return TEXT("LinkedAsset");
+	case EStateTreeStateType::Subtree:     return TEXT("Subtree");
+	default:                               return TEXT("Unknown");
+	}
+}
+
+static FString SelectionBehaviorToString(EStateTreeStateSelectionBehavior Behavior)
+{
+	switch (Behavior)
+	{
+	case EStateTreeStateSelectionBehavior::None:                                    return TEXT("None");
+	case EStateTreeStateSelectionBehavior::TryEnterState:                           return TEXT("TryEnterState");
+	case EStateTreeStateSelectionBehavior::TrySelectChildrenInOrder:                return TEXT("TrySelectChildrenInOrder");
+	case EStateTreeStateSelectionBehavior::TrySelectChildrenAtRandom:               return TEXT("TrySelectChildrenAtRandom");
+	case EStateTreeStateSelectionBehavior::TrySelectChildrenWithHighestUtility:     return TEXT("TrySelectChildrenWithHighestUtility");
+	case EStateTreeStateSelectionBehavior::TrySelectChildrenAtRandomWeightedByUtility: return TEXT("TrySelectChildrenAtRandomWeightedByUtility");
+	case EStateTreeStateSelectionBehavior::TryFollowTransitions:                    return TEXT("TryFollowTransitions");
+	default:                                                                        return TEXT("Unknown");
+	}
+}
+
+static FString TransitionTriggerToString(EStateTreeTransitionTrigger Trigger)
+{
+	TArray<FString> Flags;
+	if (EnumHasAnyFlags(Trigger, EStateTreeTransitionTrigger::OnStateSucceeded))
+	{
+		Flags.Add(TEXT("OnStateSucceeded"));
+	}
+	if (EnumHasAnyFlags(Trigger, EStateTreeTransitionTrigger::OnStateFailed))
+	{
+		Flags.Add(TEXT("OnStateFailed"));
+	}
+	if (EnumHasAnyFlags(Trigger, EStateTreeTransitionTrigger::OnTick))
+	{
+		Flags.Add(TEXT("OnTick"));
+	}
+	if (EnumHasAnyFlags(Trigger, EStateTreeTransitionTrigger::OnEvent))
+	{
+		Flags.Add(TEXT("OnEvent"));
+	}
+	if (EnumHasAnyFlags(Trigger, EStateTreeTransitionTrigger::OnDelegate))
+	{
+		Flags.Add(TEXT("OnDelegate"));
+	}
+
+	// Check for the combined OnStateCompleted (Succeeded | Failed)
+	if (Flags.Contains(TEXT("OnStateSucceeded")) && Flags.Contains(TEXT("OnStateFailed")))
+	{
+		Flags.Remove(TEXT("OnStateSucceeded"));
+		Flags.Remove(TEXT("OnStateFailed"));
+		Flags.Insert(TEXT("OnStateCompleted"), 0);
+	}
+
+	return Flags.Num() > 0 ? FString::Join(Flags, TEXT(", ")) : TEXT("None");
+}
+
+static FString TransitionPriorityToString(EStateTreeTransitionPriority Priority)
+{
+	switch (Priority)
+	{
+	case EStateTreeTransitionPriority::None:     return TEXT("None");
+	case EStateTreeTransitionPriority::Low:      return TEXT("Low");
+	case EStateTreeTransitionPriority::Normal:   return TEXT("Normal");
+	case EStateTreeTransitionPriority::Medium:   return TEXT("Medium");
+	case EStateTreeTransitionPriority::High:     return TEXT("High");
+	case EStateTreeTransitionPriority::Critical: return TEXT("Critical");
+	default:                                     return TEXT("Unknown");
+	}
+}
+
+//------------------------------------------------------------------------------
+// GetStateTreeInfo
+//------------------------------------------------------------------------------
+
+FECACommandResult FECACommand_GetStateTreeInfo::Execute(const TSharedPtr<FJsonObject>& Params)
+{
+	FString AssetPath;
+	if (!GetStringParam(Params, TEXT("asset_path"), AssetPath))
+	{
+		return FECACommandResult::Error(TEXT("Missing required parameter: asset_path"));
+	}
+
+	int32 MaxDepth = -1;
+	GetIntParam(Params, TEXT("max_depth"), MaxDepth, /*bRequired=*/false);
+
+	// Load the state tree asset
+	UStateTree* StateTree = LoadObject<UStateTree>(nullptr, *AssetPath);
+	if (!StateTree)
+	{
+		return FECACommandResult::Error(FString::Printf(TEXT("State Tree not found: %s"), *AssetPath));
+	}
+
+	TSharedPtr<FJsonObject> Result = MakeResult();
+	Result->SetStringField(TEXT("asset_path"), AssetPath);
+	Result->SetStringField(TEXT("asset_name"), StateTree->GetName());
+	Result->SetStringField(TEXT("asset_class"), TEXT("StateTree"));
+
+	// Schema
+	if (const UStateTreeSchema* Schema = StateTree->GetSchema())
+	{
+		Result->SetStringField(TEXT("schema"), Schema->GetClass()->GetName());
+	}
+
+#if WITH_EDITORONLY_DATA
+	// Access editor data for the full tree structure
+	UStateTreeEditorData* EditorData = Cast<UStateTreeEditorData>(StateTree->EditorData);
+	if (!EditorData)
+	{
+		return FECACommandResult::Error(TEXT("State Tree has no editor data (may not be compiled or is a runtime-only asset)"));
+	}
+
+	// Evaluators
+	if (EditorData->Evaluators.Num() > 0)
+	{
+		TArray<TSharedPtr<FJsonValue>> EvaluatorsArray;
+		for (const FStateTreeEditorNode& Evaluator : EditorData->Evaluators)
+		{
+			EvaluatorsArray.Add(MakeShared<FJsonValueObject>(SerializeEditorNode(Evaluator)));
+		}
+		Result->SetArrayField(TEXT("evaluators"), EvaluatorsArray);
+	}
+
+	// Global Tasks
+	if (EditorData->GlobalTasks.Num() > 0)
+	{
+		TArray<TSharedPtr<FJsonValue>> GlobalTasksArray;
+		for (const FStateTreeEditorNode& Task : EditorData->GlobalTasks)
+		{
+			GlobalTasksArray.Add(MakeShared<FJsonValueObject>(SerializeEditorNode(Task)));
+		}
+		Result->SetArrayField(TEXT("global_tasks"), GlobalTasksArray);
+	}
+
+	// State hierarchy (SubTrees are the top-level states)
+	if (EditorData->SubTrees.Num() > 0)
+	{
+		TArray<TSharedPtr<FJsonValue>> StatesArray;
+		for (const TObjectPtr<UStateTreeState>& SubTree : EditorData->SubTrees)
+		{
+			if (SubTree)
+			{
+				StatesArray.Add(MakeShared<FJsonValueObject>(SerializeState(SubTree, MaxDepth, 0)));
+			}
+		}
+		Result->SetArrayField(TEXT("states"), StatesArray);
+	}
+#else
+	return FECACommandResult::Error(TEXT("State Tree editor data is not available in non-editor builds"));
+#endif
+
+	return FECACommandResult::Success(Result);
+}
+
+//------------------------------------------------------------------------------
+// SerializeState — recursive serialization of a state and its children
+//------------------------------------------------------------------------------
+
+TSharedPtr<FJsonObject> FECACommand_GetStateTreeInfo::SerializeState(
+	const UStateTreeState* State, int32 MaxDepth, int32 CurrentDepth) const
+{
+	if (!State)
+	{
+		return MakeShared<FJsonObject>();
+	}
+
+	TSharedPtr<FJsonObject> Obj = MakeShared<FJsonObject>();
+
+	Obj->SetStringField(TEXT("name"), State->Name.ToString());
+	Obj->SetStringField(TEXT("type"), StateTypeToString(State->Type));
+	Obj->SetStringField(TEXT("id"), State->ID.ToString());
+	Obj->SetBoolField(TEXT("enabled"), State->bEnabled);
+	Obj->SetStringField(TEXT("selection_behavior"), SelectionBehaviorToString(State->SelectionBehavior));
+	Obj->SetNumberField(TEXT("depth"), CurrentDepth);
+
+	// GameplayTag
+	if (State->Tag.IsValid())
+	{
+		Obj->SetStringField(TEXT("tag"), State->Tag.ToString());
+	}
+
+	// Description
+	if (!State->Description.IsEmpty())
+	{
+		Obj->SetStringField(TEXT("description"), State->Description);
+	}
+
+	// Linked state/asset info
+	if (State->Type == EStateTreeStateType::LinkedAsset && State->LinkedAsset)
+	{
+		Obj->SetStringField(TEXT("linked_asset"), State->LinkedAsset->GetPathName());
+	}
+#if WITH_EDITORONLY_DATA
+	if (State->Type == EStateTreeStateType::Linked && State->LinkedSubtree.Name != NAME_None)
+	{
+		Obj->SetStringField(TEXT("linked_subtree"), State->LinkedSubtree.Name.ToString());
+	}
+#endif
+
+	// Required event to enter
+	if (State->bHasRequiredEventToEnter && State->RequiredEventToEnter.IsValid())
+	{
+		TSharedPtr<FJsonObject> EventObj = MakeShared<FJsonObject>();
+		if (State->RequiredEventToEnter.Tag.IsValid())
+		{
+			EventObj->SetStringField(TEXT("tag"), State->RequiredEventToEnter.Tag.ToString());
+		}
+		if (State->RequiredEventToEnter.PayloadStruct)
+		{
+			EventObj->SetStringField(TEXT("payload_type"), State->RequiredEventToEnter.PayloadStruct->GetName());
+		}
+		Obj->SetObjectField(TEXT("required_event"), EventObj);
+	}
+
+	// Weight (for utility-based selection)
+	if (State->Weight != 1.0f)
+	{
+		Obj->SetNumberField(TEXT("weight"), State->Weight);
+	}
+
+	// Custom tick rate
+	if (State->bHasCustomTickRate)
+	{
+		Obj->SetNumberField(TEXT("custom_tick_rate"), State->CustomTickRate);
+	}
+
+	// Enter Conditions
+	if (State->EnterConditions.Num() > 0)
+	{
+		TArray<TSharedPtr<FJsonValue>> ConditionsArray;
+		for (const FStateTreeEditorNode& Condition : State->EnterConditions)
+		{
+			ConditionsArray.Add(MakeShared<FJsonValueObject>(SerializeEditorNode(Condition)));
+		}
+		Obj->SetArrayField(TEXT("enter_conditions"), ConditionsArray);
+	}
+
+	// Utility Considerations
+	if (State->Considerations.Num() > 0)
+	{
+		TArray<TSharedPtr<FJsonValue>> ConsiderationsArray;
+		for (const FStateTreeEditorNode& Consideration : State->Considerations)
+		{
+			ConsiderationsArray.Add(MakeShared<FJsonValueObject>(SerializeEditorNode(Consideration)));
+		}
+		Obj->SetArrayField(TEXT("considerations"), ConsiderationsArray);
+	}
+
+	// Tasks
+	if (State->Tasks.Num() > 0)
+	{
+		TArray<TSharedPtr<FJsonValue>> TasksArray;
+		for (const FStateTreeEditorNode& Task : State->Tasks)
+		{
+			TasksArray.Add(MakeShared<FJsonValueObject>(SerializeEditorNode(Task)));
+		}
+		Obj->SetArrayField(TEXT("tasks"), TasksArray);
+	}
+
+	// Single Task (used when schema calls for single task per state)
+	if (State->SingleTask.Node.IsValid())
+	{
+		Obj->SetObjectField(TEXT("task"), SerializeEditorNode(State->SingleTask));
+	}
+
+	// Transitions
+	if (State->Transitions.Num() > 0)
+	{
+		TArray<TSharedPtr<FJsonValue>> TransitionsArray;
+		for (const FStateTreeTransition& Transition : State->Transitions)
+		{
+			TransitionsArray.Add(MakeShared<FJsonValueObject>(SerializeTransition(Transition)));
+		}
+		Obj->SetArrayField(TEXT("transitions"), TransitionsArray);
+	}
+
+	// Children (recursive)
+	bool bShouldRecurse = (MaxDepth < 0 || CurrentDepth < MaxDepth);
+	if (State->Children.Num() > 0)
+	{
+		TArray<TSharedPtr<FJsonValue>> ChildrenArray;
+		for (const TObjectPtr<UStateTreeState>& Child : State->Children)
+		{
+			if (Child)
+			{
+				if (bShouldRecurse)
+				{
+					ChildrenArray.Add(MakeShared<FJsonValueObject>(SerializeState(Child, MaxDepth, CurrentDepth + 1)));
+				}
+				else
+				{
+					// Hit depth limit — just show name and child count
+					TSharedPtr<FJsonObject> ShallowObj = MakeShared<FJsonObject>();
+					ShallowObj->SetStringField(TEXT("name"), Child->Name.ToString());
+					ShallowObj->SetStringField(TEXT("type"), StateTypeToString(Child->Type));
+					ShallowObj->SetBoolField(TEXT("truncated"), true);
+					ShallowObj->SetNumberField(TEXT("child_count"), Child->Children.Num());
+					ChildrenArray.Add(MakeShared<FJsonValueObject>(ShallowObj));
+				}
+			}
+		}
+		Obj->SetArrayField(TEXT("children"), ChildrenArray);
+	}
+
+	return Obj;
+}
+
+//------------------------------------------------------------------------------
+// SerializeEditorNode — serialize a task, condition, evaluator, or consideration
+//------------------------------------------------------------------------------
+
+TSharedPtr<FJsonObject> FECACommand_GetStateTreeInfo::SerializeEditorNode(
+	const FStateTreeEditorNode& EditorNode) const
+{
+	TSharedPtr<FJsonObject> Obj = MakeShared<FJsonObject>();
+
+	// Node name
+	FName NodeName = EditorNode.GetName();
+	if (NodeName != NAME_None)
+	{
+		Obj->SetStringField(TEXT("name"), NodeName.ToString());
+	}
+
+	// Node struct type (e.g. "FStateTreeRunBehaviorTask", "FStateTreeCompareIntCondition")
+	if (EditorNode.Node.IsValid())
+	{
+		const UScriptStruct* NodeStruct = EditorNode.Node.GetScriptStruct();
+		if (NodeStruct)
+		{
+			Obj->SetStringField(TEXT("node_type"), NodeStruct->GetName());
+		}
+	}
+
+	// Instance data type
+	if (EditorNode.Instance.IsValid())
+	{
+		const UScriptStruct* InstanceStruct = EditorNode.Instance.GetScriptStruct();
+		if (InstanceStruct)
+		{
+			Obj->SetStringField(TEXT("instance_type"), InstanceStruct->GetName());
+		}
+	}
+	else if (EditorNode.InstanceObject)
+	{
+		Obj->SetStringField(TEXT("instance_type"), EditorNode.InstanceObject->GetClass()->GetName());
+	}
+
+	// ID
+	Obj->SetStringField(TEXT("id"), EditorNode.ID.ToString());
+
+	// Expression operand (for conditions: And/Or)
+	if (EditorNode.ExpressionOperand == EStateTreeExpressionOperand::Or)
+	{
+		Obj->SetStringField(TEXT("operand"), TEXT("Or"));
+	}
+	else if (EditorNode.ExpressionIndent > 0 || EditorNode.ExpressionOperand == EStateTreeExpressionOperand::And)
+	{
+		Obj->SetStringField(TEXT("operand"), TEXT("And"));
+	}
+
+	if (EditorNode.ExpressionIndent > 0)
+	{
+		Obj->SetNumberField(TEXT("indent"), EditorNode.ExpressionIndent);
+	}
+
+	return Obj;
+}
+
+//------------------------------------------------------------------------------
+// SerializeTransition — serialize a state transition
+//------------------------------------------------------------------------------
+
+TSharedPtr<FJsonObject> FECACommand_GetStateTreeInfo::SerializeTransition(
+	const FStateTreeTransition& Transition) const
+{
+	TSharedPtr<FJsonObject> Obj = MakeShared<FJsonObject>();
+
+	// Trigger type
+	Obj->SetStringField(TEXT("trigger"), TransitionTriggerToString(Transition.Trigger));
+
+	// Target state
+#if WITH_EDITORONLY_DATA
+	if (Transition.State.Name != NAME_None)
+	{
+		Obj->SetStringField(TEXT("target_state"), Transition.State.Name.ToString());
+	}
+	if (Transition.State.LinkType != EStateTreeTransitionType::None)
+	{
+		FString LinkTypeStr;
+		switch (Transition.State.LinkType)
+		{
+		case EStateTreeTransitionType::Succeeded:            LinkTypeStr = TEXT("Succeeded"); break;
+		case EStateTreeTransitionType::Failed:               LinkTypeStr = TEXT("Failed"); break;
+		case EStateTreeTransitionType::GotoState:            LinkTypeStr = TEXT("GotoState"); break;
+		case EStateTreeTransitionType::Parent:               LinkTypeStr = TEXT("Parent"); break;
+		case EStateTreeTransitionType::NextState:            LinkTypeStr = TEXT("NextState"); break;
+		case EStateTreeTransitionType::NextSelectableState:  LinkTypeStr = TEXT("NextSelectableState"); break;
+		case EStateTreeTransitionType::NextParent:           LinkTypeStr = TEXT("NextParent"); break;
+		case EStateTreeTransitionType::NextSelectableParent: LinkTypeStr = TEXT("NextSelectableParent"); break;
+		default:                                             LinkTypeStr = TEXT("Unknown"); break;
+		}
+		Obj->SetStringField(TEXT("transition_type"), LinkTypeStr);
+	}
+#endif
+
+	// Priority
+	Obj->SetStringField(TEXT("priority"), TransitionPriorityToString(Transition.Priority));
+
+	// Required event
+	if (Transition.RequiredEvent.IsValid())
+	{
+		TSharedPtr<FJsonObject> EventObj = MakeShared<FJsonObject>();
+		if (Transition.RequiredEvent.Tag.IsValid())
+		{
+			EventObj->SetStringField(TEXT("tag"), Transition.RequiredEvent.Tag.ToString());
+		}
+		if (Transition.RequiredEvent.PayloadStruct)
+		{
+			EventObj->SetStringField(TEXT("payload_type"), Transition.RequiredEvent.PayloadStruct->GetName());
+		}
+		Obj->SetObjectField(TEXT("event"), EventObj);
+	}
+
+	// Delay
+	if (Transition.bDelayTransition)
+	{
+		Obj->SetNumberField(TEXT("delay"), Transition.DelayDuration);
+		if (Transition.DelayRandomVariance > 0.0f)
+		{
+			Obj->SetNumberField(TEXT("delay_random_variance"), Transition.DelayRandomVariance);
+		}
+	}
+
+	// Enabled
+	if (!Transition.bTransitionEnabled)
+	{
+		Obj->SetBoolField(TEXT("enabled"), false);
+	}
+
+	// Conditions
+	if (Transition.Conditions.Num() > 0)
+	{
+		TArray<TSharedPtr<FJsonValue>> ConditionsArray;
+		for (const FStateTreeEditorNode& Condition : Transition.Conditions)
+		{
+			ConditionsArray.Add(MakeShared<FJsonValueObject>(SerializeEditorNode(Condition)));
+		}
+		Obj->SetArrayField(TEXT("conditions"), ConditionsArray);
+	}
+
+	// ID
+	Obj->SetStringField(TEXT("id"), Transition.ID.ToString());
+
+	return Obj;
+}
diff --git a/Engine/Plugins/Experimental/ECABridge/Source/ECABridge/Public/Commands/ECABehaviorTreeCommands.h b/Engine/Plugins/Experimental/ECABridge/Source/ECABridge/Public/Commands/ECABehaviorTreeCommands.h
new file mode 100644
--- /dev/null
+++ b/Engine/Plugins/Experimental/ECABridge/Source/ECABridge/Public/Commands/ECABehaviorTreeCommands.h
@@ -0,0 +1,63 @@
+// Copyright Epic Games, Inc. All Rights Reserved.
+
+#pragma once
+
+#include "Commands/ECACommand.h"
+
+class UBTNode;
+class UBTCompositeNode;
+class UBlackboardData;
+
+/**
+ * Get full node hierarchy of a Behavior Tree asset
+ */
+class FECACommand_GetBehaviorTreeInfo : public IECACommand
+{
+public:
+	virtual FString GetName() const override { return TEXT("get_behavior_tree_info"); }
+	virtual FString GetDescription() const override
+	{
+		return TEXT("Get the full node hierarchy of a Behavior Tree asset including composites "
+			"(selector/sequence/parallel), tasks, decorators, services, and the associated "
+			"blackboard keys. Returns structured JSON tree.");
+	}
+	virtual FString GetCategory() const override { return TEXT("AI"); }
+
+	virtual TArray<FECACommandParam> GetParameters() const override
+	{
+		return {
+			{ TEXT("asset_path"), TEXT("string"), TEXT("Content path to the Behavior Tree asset (e.g., /Game/AI/BT_MyTree)"), true },
+			{ TEXT("max_depth"), TEXT("number"), TEXT("Maximum recursion depth for the node tree (-1 for unlimited, default)"), false, TEXT("-1") },
+		};
+	}
+
+	virtual FECACommandResult Execute(const TSharedPtr<FJsonObject>& Params) override;
+
+private:
+	TSharedPtr<FJsonObject> SerializeNode(UBTNode* Node) const;
+	TSharedPtr<FJsonObject> SerializeCompositeNode(UBTCompositeNode* CompositeNode, int32 MaxDepth, int32 CurrentDepth) const;
+};
+
+/**
+ * Get key definitions of a Blackboard asset independently of a Behavior Tree
+ */
+class FECACommand_GetBlackboardInfo : public IECACommand
+{
+public:
+	virtual FString GetName() const override { return TEXT("get_blackboard_info"); }
+	virtual FString GetDescription() const override
+	{
+		return TEXT("Get the key definitions of a Blackboard asset including key names, "
+			"types, categories, descriptions, instance sync status, and parent chain.");
+	}
+	virtual FString GetCategory() const override { return TEXT("AI"); }
+
+	virtual TArray<FECACommandParam> GetParameters() const override
+	{
+		return {
+			{ TEXT("asset_path"), TEXT("string"), TEXT("Content path to the Blackboard asset (e.g., /Game/AI/BB_MyBlackboard)"), true },
+		};
+	}
+
+	virtual FECACommandResult Execute(const TSharedPtr<FJsonObject>& Params) override;
+};
\ No newline at end of file
diff --git a/Engine/Plugins/Experimental/ECABridge/Source/ECABridge/Public/Commands/ECAEQSCommands.h b/Engine/Plugins/Experimental/ECABridge/Source/ECABridge/Public/Commands/ECAEQSCommands.h
new file mode 100644
--- /dev/null
+++ b/Engine/Plugins/Experimental/ECABridge/Source/ECABridge/Public/Commands/ECAEQSCommands.h
@@ -0,0 +1,39 @@
+// Copyright Epic Games, Inc. All Rights Reserved.
+
+#pragma once
+
+#include "Commands/ECACommand.h"
+
+class UEnvQueryOption;
+class UEnvQueryGenerator;
+class UEnvQueryTest;
+
+/**
+ * Get full structure of an Environment Query (EQS) asset
+ */
+class FECACommand_GetEQSInfo : public IECACommand
+{
+public:
+	virtual FString GetName() const override { return TEXT("get_eqs_info"); }
+	virtual FString GetDescription() const override
+	{
+		return TEXT("Get the full structure of an Environment Query (EQS) asset including "
+			"generators, tests, scoring equations, filter settings, and test ordering. "
+			"Returns structured JSON with human-readable descriptions.");
+	}
+	virtual FString GetCategory() const override { return TEXT("AI"); }
+
+	virtual TArray<FECACommandParam> GetParameters() const override
+	{
+		return {
+			{ TEXT("asset_path"), TEXT("string"), TEXT("Content path to the EQS asset (e.g., /Game/AI/EQS_FindCover)"), true },
+		};
+	}
+
+	virtual FECACommandResult Execute(const TSharedPtr<FJsonObject>& Params) override;
+
+private:
+	TSharedPtr<FJsonObject> SerializeOption(const UEnvQueryOption* Option) const;
+	TSharedPtr<FJsonObject> SerializeGenerator(const UEnvQueryGenerator* Generator) const;
+	TSharedPtr<FJsonObject> SerializeTest(const UEnvQueryTest* Test) const;
+};
diff --git a/Engine/Plugins/Experimental/ECABridge/Source/ECABridge/Public/Commands/ECAStateTreeCommands.h b/Engine/Plugins/Experimental/ECABridge/Source/ECABridge/Public/Commands/ECAStateTreeCommands.h
new file mode 100644
--- /dev/null
+++ b/Engine/Plugins/Experimental/ECABridge/Source/ECABridge/Public/Commands/ECAStateTreeCommands.h
@@ -0,0 +1,40 @@
+// Copyright Epic Games, Inc. All Rights Reserved.
+
+#pragma once
+
+#include "Commands/ECACommand.h"
+
+class UStateTreeState;
+struct FStateTreeEditorNode;
+struct FStateTreeTransition;
+
+/**
+ * Get full hierarchy of a State Tree asset
+ */
+class FECACommand_GetStateTreeInfo : public IECACommand
+{
+public:
+	virtual FString GetName() const override { return TEXT("get_state_tree_info"); }
+	virtual FString GetDescription() const override
+	{
+		return TEXT("Get the full hierarchy of a State Tree asset including states, tasks, "
+			"enter conditions, considerations, transitions, evaluators, and global tasks. "
+			"Returns structured JSON tree matching the editor representation.");
+	}
+	virtual FString GetCategory() const override { return TEXT("AI"); }
+
+	virtual TArray<FECACommandParam> GetParameters() const override
+	{
+		return {
+			{ TEXT("asset_path"), TEXT("string"), TEXT("Content path to the State Tree asset (e.g., /Game/AI/ST_MyStateTree)"), true },
+			{ TEXT("max_depth"), TEXT("number"), TEXT("Maximum recursion depth for the state hierarchy (-1 for unlimited, default)"), false, TEXT("-1") },
+		};
+	}
+
+	virtual FECACommandResult Execute(const TSharedPtr<FJsonObject>& Params) override;
+
+private:
+	TSharedPtr<FJsonObject> SerializeState(const UStateTreeState* State, int32 MaxDepth, int32 CurrentDepth) const;
+	TSharedPtr<FJsonObject> SerializeEditorNode(const FStateTreeEditorNode& EditorNode) const;
+	TSharedPtr<FJsonObject> SerializeTransition(const FStateTreeTransition& Transition) const;
+};
