[ECABridge] Fix save_asset crash on modified materials

save_asset, save_all_dirty, and save_and_compile_blueprint now use
UEditorLoadingAndSavingUtils::SavePackages (full Ctrl+S pipeline)
instead of bare UPackage::SavePackage. The old path crashed in
FinalizeFile when saving materials with modified graph nodes because
the old file loader still held a read handle on the .uasset.

#rnx
#Synced-CL 52047883

diff --git a/Engine/Plugins/Experimental/ECABridge/Source/ECABridge/Private/Commands/ECASaveCommands.cpp b/Engine/Plugins/Experimental/ECABridge/Source/ECABridge/Private/Commands/ECASaveCommands.cpp
--- a/Engine/Plugins/Experimental/ECABridge/Source/ECABridge/Private/Commands/ECASaveCommands.cpp
+++ b/Engine/Plugins/Experimental/ECABridge/Source/ECABridge/Private/Commands/ECASaveCommands.cpp
@@ -7,6 +7,8 @@
 #include "UObject/Package.h"
 #include "UObject/SavePackage.h"
 #include "UObject/UObjectIterator.h"
+#include "UObject/Linker.h"
+#include "FileHelpers.h"
 
 // Self-registration — no other files need modification
 REGISTER_ECA_COMMAND(FECACommand_IsAssetDirty)
@@ -54,6 +56,12 @@ namespace ECASaveHelpers
 		FString PackageFileName = FPackageName::LongPackageNameToFilename(
 			Package->GetName(), GetPackageExtension(Package));
 
+		// Detach any active file loaders before writing so SavePackage can
+		// open the file for writing. Without this, materials with modified
+		// graph nodes crash in FinalizeFile because the old loader still
+		// holds a read handle on the .uasset file.
+		ResetLoaders(Package);
+
 		FSavePackageArgs SaveArgs;
 		SaveArgs.TopLevelFlags = RF_Standalone;
 		bool bSaved = UPackage::SavePackage(Package, Asset, *PackageFileName, SaveArgs);
@@ -132,7 +140,13 @@ FECACommandResult FECACommand_SaveAsset::Execute(const TSharedPtr<FJsonObject>&
 
 	bool bWasDirty = Resolved.Package->IsDirty();
 
-	if (!ECASaveHelpers::SavePackageToDisk(Resolved.Package, Resolved.Asset))
+	// Use the full editor save pipeline (same as Ctrl+S) which calls
+	// EnsureLoadingComplete + ResetLoaders before SavePackage. The bare
+	// UPackage::SavePackage call crashes on materials with modified graphs
+	// because the old file loader still holds a read handle on the .uasset.
+	TArray<UPackage*> PackagesToSave;
+	PackagesToSave.Add(Resolved.Package);
+	if (!UEditorLoadingAndSavingUtils::SavePackages(PackagesToSave, /*bOnlyDirty=*/false))
 	{
 		return FECACommandResult::Error(FString::Printf(
 			TEXT("Failed to save package: %s"), *Resolved.Package->GetName()));
@@ -168,7 +182,9 @@ FECACommandResult FECACommand_SaveAllDirty::Execute(const TSharedPtr<FJsonObject
 
 		FString PackageName = Package->GetName();
 
-		if (ECASaveHelpers::SavePackageToDisk(Package, nullptr))
+		TArray<UPackage*> SinglePackage;
+		SinglePackage.Add(Package);
+		if (UEditorLoadingAndSavingUtils::SavePackages(SinglePackage, /*bOnlyDirty=*/false))
 		{
 			SavedCount++;
 			SavedArray.Add(MakeShared<FJsonValueString>(PackageName));
@@ -216,7 +232,9 @@ FECACommandResult FECACommand_SaveAndCompileBlueprint::Execute(const TSharedPtr<
 
 	// Step 2: Save (always — even with compile errors, persist WIP)
 	UPackage* Package = Blueprint->GetPackage();
-	bool bSaved = ECASaveHelpers::SavePackageToDisk(Package, Blueprint);
+	TArray<UPackage*> PackagesToSave;
+	PackagesToSave.Add(Package);
+	bool bSaved = UEditorLoadingAndSavingUtils::SavePackages(PackagesToSave, /*bOnlyDirty=*/false);
 
 	TSharedPtr<FJsonObject> Result = MakeResult();
 	Result->SetBoolField(TEXT("compiled"), true);
