ECABridge: make ToolsetRegistry an optional dep

#Synced-CL 52144611

diff --git a/Engine/Plugins/Experimental/ECABridge/ECABridge.uplugin b/Engine/Plugins/Experimental/ECABridge/ECABridge.uplugin
--- a/Engine/Plugins/Experimental/ECABridge/ECABridge.uplugin
+++ b/Engine/Plugins/Experimental/ECABridge/ECABridge.uplugin
@@ -25,7 +25,8 @@
 	"Plugins": [
 		{
 			"Name": "ToolsetRegistry",
-			"Enabled": true
+			"Enabled": true,
+			"Optional": true
 		},
 		{
 			"Name": "Niagara",
@@ -54,6 +55,10 @@
 		{
 			"Name": "ProceduralMeshComponent",
 			"Enabled": true
+		},
+		{
+			"Name": "StateTree",
+			"Enabled": true
 		}
 	]
 }
diff --git a/Engine/Plugins/Experimental/ECABridge/Source/ECABridge/ECABridge.Build.cs b/Engine/Plugins/Experimental/ECABridge/Source/ECABridge/ECABridge.Build.cs
--- a/Engine/Plugins/Experimental/ECABridge/Source/ECABridge/ECABridge.Build.cs
+++ b/Engine/Plugins/Experimental/ECABridge/Source/ECABridge/ECABridge.Build.cs
@@ -64,9 +64,6 @@ public class ECABridge : ModuleRules
 			// Additional utilities
 			"Projects", "DeveloperSettings", "ToolMenus",
 			
-			// ToolsetRegistry integration (bidirectional: expose ECA commands as toolset + consume toolsets via MCP)
-			"ToolsetRegistry",
-			
 			// Image processing
 			"ImageWrapper",
 			
@@ -83,5 +80,17 @@ public class ECABridge : ModuleRules
 		// when the plugin is EnabledByDefault=false, so add it explicitly.
 		string MVVMBlueprintPath = System.IO.Path.Combine(EngineDirectory, "Plugins", "Runtime", "ModelViewViewModel", "Source", "ModelViewViewModelBlueprint", "Public");
 		PrivateIncludePaths.Add(MVVMBlueprintPath);
+
+		// ToolsetRegistry integration is optional (bidirectional: expose ECA commands as toolset + consume toolsets via MCP)
+		bool bHasToolsetRegistry = Directory.Exists(Path.Combine(EngineDirectory, "Plugins", "Experimental", "ToolsetRegistry", "Source"));
+		if (bHasToolsetRegistry)
+		{
+			PrivateDependencyModuleNames.Add("ToolsetRegistry");
+			PublicDefinitions.Add("WITH_TOOLSET_REGISTRY=1");
+		}
+		else
+		{
+			PublicDefinitions.Add("WITH_TOOLSET_REGISTRY=0");
+		}
 	}
 }
diff --git a/Engine/Plugins/Experimental/ECABridge/Source/ECABridge/Private/Commands/ECAToolsetRegistryCommands.cpp b/Engine/Plugins/Experimental/ECABridge/Source/ECABridge/Private/Commands/ECAToolsetRegistryCommands.cpp
new file mode 100644
--- /dev/null
+++ b/Engine/Plugins/Experimental/ECABridge/Source/ECABridge/Private/Commands/ECAToolsetRegistryCommands.cpp
@@ -0,0 +1,588 @@
+// Copyright Epic Games, Inc. All Rights Reserved.
+
+#include "Commands/ECACommand.h"
+#include "Dom/JsonObject.h"
+#include "Dom/JsonValue.h"
+#include "Serialization/JsonSerializer.h"
+#include "Serialization/JsonWriter.h"
+#include "Serialization/JsonReader.h"
+#include "Misc/FileHelper.h"
+#include "Misc/Paths.h"
+#include "HAL/PlatformFileManager.h"
+#include "HAL/PlatformProcess.h"
+
+//------------------------------------------------------------------------------
+// Toolset Registry State
+// 
+// The ECABridge C++ side maintains a cached view of the Python-side function
+// registry. This cache is populated by ECA calling store_registered_functions
+// (which syncs from the Python registry), or by the bridge querying directly.
+//
+// For function calls, the bridge queues a request that ECA picks up via the
+// event system and returns results through the normal MCP response flow.
+//------------------------------------------------------------------------------
+
+namespace ToolsetRegistryCache
+{
+	struct FRegisteredFunctionInfo
+	{
+		FString Name;
+		FString Description;
+		FString Category;
+		FString Scope;
+		FString ReturnType;
+		FString SourceCode;
+		TArray<TSharedPtr<FJsonObject>> Parameters;
+	};
+
+	static TMap<FString, FRegisteredFunctionInfo> CachedFunctions;
+	static FCriticalSection CacheLock;
+
+	void UpdateCache(const TArray<TSharedPtr<FJsonValue>>& Functions)
+	{
+		FScopeLock Lock(&CacheLock);
+		CachedFunctions.Empty();
+
+		for (const TSharedPtr<FJsonValue>& FuncValue : Functions)
+		{
+			const TSharedPtr<FJsonObject>& FuncObj = FuncValue->AsObject();
+			if (!FuncObj.IsValid())
+			{
+				continue;
+			}
+
+			FRegisteredFunctionInfo Info;
+			FuncObj->TryGetStringField(TEXT("name"), Info.Name);
+			FuncObj->TryGetStringField(TEXT("description"), Info.Description);
+			FuncObj->TryGetStringField(TEXT("category"), Info.Category);
+			FuncObj->TryGetStringField(TEXT("scope"), Info.Scope);
+			FuncObj->TryGetStringField(TEXT("return_type"), Info.ReturnType);
+			FuncObj->TryGetStringField(TEXT("source"), Info.SourceCode);
+
+			const TArray<TSharedPtr<FJsonValue>>* ParamsArray;
+			if (FuncObj->TryGetArrayField(TEXT("parameters"), ParamsArray))
+			{
+				for (const TSharedPtr<FJsonValue>& ParamVal : *ParamsArray)
+				{
+					if (ParamVal->Type == EJson::Object)
+					{
+						Info.Parameters.Add(ParamVal->AsObject());
+					}
+				}
+			}
+
+			if (!Info.Name.IsEmpty())
+			{
+				CachedFunctions.Add(Info.Name, MoveTemp(Info));
+			}
+		}
+
+		UE_LOG(LogTemp, Log, TEXT("[ECABridge] Toolset registry cache updated with %d functions"), CachedFunctions.Num());
+	}
+
+	TArray<FRegisteredFunctionInfo> GetAll(const FString& ScopeFilter = TEXT(""), const FString& CategoryFilter = TEXT(""))
+	{
+		FScopeLock Lock(&CacheLock);
+		TArray<FRegisteredFunctionInfo> Result;
+
+		for (const auto& Pair : CachedFunctions)
+		{
+			const FRegisteredFunctionInfo& Info = Pair.Value;
+
+			if (!ScopeFilter.IsEmpty() && Info.Scope != ScopeFilter)
+			{
+				continue;
+			}
+			if (!CategoryFilter.IsEmpty() && Info.Category != CategoryFilter)
+			{
+				continue;
+			}
+
+			Result.Add(Info);
+		}
+
+		return Result;
+	}
+
+	const FRegisteredFunctionInfo* Find(const FString& Name)
+	{
+		FScopeLock Lock(&CacheLock);
+		return CachedFunctions.Find(Name);
+	}
+}
+
+//------------------------------------------------------------------------------
+// store_registered_functions - Internal command for ECA to sync registry state
+// 
+// ECA calls this after loading/modifying the function registry to keep the
+// bridge's cache up to date. This is the primary sync mechanism.
+//------------------------------------------------------------------------------
+
+class FECACommand_StoreRegisteredFunctions : public IECACommand
+{
+public:
+	virtual FString GetName() const override { return TEXT("store_registered_functions"); }
+	virtual FString GetDescription() const override 
+	{ 
+		return TEXT("Internal: Sync the registered functions cache from the ECA Python side. Called automatically by ECA when the registry changes."); 
+	}
+	virtual FString GetCategory() const override { return TEXT("ToolsetRegistry"); }
+	
+	virtual TArray<FECACommandParam> GetParameters() const override
+	{
+		return {
+			{ TEXT("functions"), TEXT("array"), TEXT("Array of registered function info objects"), true }
+		};
+	}
+
+	virtual FECACommandResult Execute(const TSharedPtr<FJsonObject>& Params) override
+	{
+		const TArray<TSharedPtr<FJsonValue>>* FunctionsArray;
+		if (!Params->TryGetArrayField(TEXT("functions"), FunctionsArray))
+		{
+			return FECACommandResult::Error(TEXT("Missing required parameter: functions (array)"));
+		}
+
+		ToolsetRegistryCache::UpdateCache(*FunctionsArray);
+
+		TSharedPtr<FJsonObject> Result = MakeResult();
+		Result->SetBoolField(TEXT("success"), true);
+		Result->SetNumberField(TEXT("cached_count"), FunctionsArray->Num());
+
+		return FECACommandResult::Success(Result);
+	}
+};
+
+//------------------------------------------------------------------------------
+// list_registered_functions
+//------------------------------------------------------------------------------
+
+class FECACommand_ListRegisteredFunctions : public IECACommand
+{
+public:
+	virtual FString GetName() const override { return TEXT("list_registered_functions"); }
+	virtual FString GetDescription() const override { return TEXT("List all registered functions from the toolset registry cache, with optional scope/category filtering."); }
+	virtual FString GetCategory() const override { return TEXT("ToolsetRegistry"); }
+
+	virtual TArray<FECACommandParam> GetParameters() const override
+	{
+		return {
+			{ TEXT("scope"), TEXT("string"), TEXT("Filter by scope (e.g. 'project', 'global')"), false },
+			{ TEXT("category"), TEXT("string"), TEXT("Filter by category"), false }
+		};
+	}
+
+	virtual FECACommandResult Execute(const TSharedPtr<FJsonObject>& Params) override;
+};
+
+FECACommandResult FECACommand_ListRegisteredFunctions::Execute(const TSharedPtr<FJsonObject>& Params)
+{
+	FString ScopeFilter;
+	GetStringParam(Params, TEXT("scope"), ScopeFilter, false);
+
+	FString CategoryFilter;
+	GetStringParam(Params, TEXT("category"), CategoryFilter, false);
+
+	TArray<ToolsetRegistryCache::FRegisteredFunctionInfo> Functions = 
+		ToolsetRegistryCache::GetAll(ScopeFilter, CategoryFilter);
+
+	TSharedPtr<FJsonObject> Result = MakeResult();
+	
+	TArray<TSharedPtr<FJsonValue>> FunctionsArray;
+	for (const auto& Info : Functions)
+	{
+		TSharedPtr<FJsonObject> FuncObj = MakeShared<FJsonObject>();
+		FuncObj->SetStringField(TEXT("name"), Info.Name);
+		FuncObj->SetStringField(TEXT("description"), Info.Description);
+		FuncObj->SetStringField(TEXT("category"), Info.Category);
+		FuncObj->SetStringField(TEXT("scope"), Info.Scope);
+		FuncObj->SetStringField(TEXT("return_type"), Info.ReturnType);
+		
+		// Include parameter info
+		TArray<TSharedPtr<FJsonValue>> ParamsArray;
+		for (const TSharedPtr<FJsonObject>& ParamObj : Info.Parameters)
+		{
+			ParamsArray.Add(MakeShared<FJsonValueObject>(ParamObj));
+		}
+		FuncObj->SetArrayField(TEXT("parameters"), ParamsArray);
+		
+		FunctionsArray.Add(MakeShared<FJsonValueObject>(FuncObj));
+	}
+
+	Result->SetArrayField(TEXT("functions"), FunctionsArray);
+	Result->SetNumberField(TEXT("count"), Functions.Num());
+
+	return FECACommandResult::Success(Result);
+}
+
+//------------------------------------------------------------------------------
+// call_registered_function
+//
+// This is the key integration point. When Unreal-side code (or the AI via
+// ECABridge) needs to call a registered Python function, this command
+// queues a "call_function" event. ECA's event polling picks it up,
+// executes the function, and pushes the result back via a separate
+// "function_result" event or direct MCP response.
+//
+// For the synchronous MCP call path (AI calls tools/call with
+// call_registered_function), we serialize the request as JSON and ECA
+// executes it inline via meta_tool, returning the result directly.
+//------------------------------------------------------------------------------
+
+class FECACommand_CallRegisteredFunction : public IECACommand
+{
+public:
+	virtual FString GetName() const override { return TEXT("call_registered_function"); }
+	virtual FString GetDescription() const override { return TEXT("Call a registered Python function by name. The request is forwarded to ECA for execution."); }
+	virtual FString GetCategory() const override { return TEXT("ToolsetRegistry"); }
+
+	virtual TArray<FECACommandParam> GetParameters() const override
+	{
+		return {
+			{ TEXT("function_name"), TEXT("string"), TEXT("Name of the registered function to call"), true },
+			{ TEXT("arguments"), TEXT("object"), TEXT("Arguments to pass to the function"), false }
+		};
+	}
+
+	virtual FECACommandResult Execute(const TSharedPtr<FJsonObject>& Params) override;
+};
+
+FECACommandResult FECACommand_CallRegisteredFunction::Execute(const TSharedPtr<FJsonObject>& Params)
+{
+	FString FunctionName;
+	if (!GetStringParam(Params, TEXT("function_name"), FunctionName))
+	{
+		return FECACommandResult::Error(TEXT("Missing required parameter: function_name"));
+	}
+
+	// Verify the function exists in our cache
+	const ToolsetRegistryCache::FRegisteredFunctionInfo* FuncInfo = ToolsetRegistryCache::Find(FunctionName);
+	if (!FuncInfo)
+	{
+		return FECACommandResult::Error(FString::Printf(
+			TEXT("Function '%s' not found in toolset registry cache. Call store_registered_functions first to sync the cache."),
+			*FunctionName));
+	}
+
+	// Build the call request
+	TSharedPtr<FJsonObject> Arguments = MakeShared<FJsonObject>();
+	const TSharedPtr<FJsonObject>* ArgsPtr;
+	if (Params->TryGetObjectField(TEXT("arguments"), ArgsPtr))
+	{
+		Arguments = *ArgsPtr;
+	}
+
+	// Build a meta_tool code snippet that calls the registered function
+	// This will be executed by ECA when it processes the MCP response
+	FString ArgsJson;
+	{
+		TSharedRef<FJsonObject> ArgsRef = Arguments.ToSharedRef();
+		TSharedRef<TJsonWriter<>> Writer = TJsonWriterFactory<>::Create(&ArgsJson);
+		FJsonSerializer::Serialize(ArgsRef, Writer);
+	}
+
+	TSharedPtr<FJsonObject> Result = MakeResult();
+	Result->SetStringField(TEXT("function_name"), FunctionName);
+	Result->SetObjectField(TEXT("arguments"), Arguments);
+	Result->SetStringField(TEXT("description"), FuncInfo->Description);
+	Result->SetStringField(TEXT("category"), FuncInfo->Category);
+	Result->SetStringField(TEXT("status"), TEXT("call_requested"));
+	Result->SetStringField(TEXT("instruction"), 
+		FString::Printf(TEXT("Execute this registered function via meta_tool: tools.%s(**%s)"), 
+			*FunctionName, *ArgsJson));
+
+	return FECACommandResult::Success(Result);
+}
+
+//------------------------------------------------------------------------------
+// get_registered_function_info
+//------------------------------------------------------------------------------
+
+class FECACommand_GetRegisteredFunctionInfo : public IECACommand
+{
+public:
+	virtual FString GetName() const override { return TEXT("get_registered_function_info"); }
+	virtual FString GetDescription() const override { return TEXT("Get detailed info about a specific registered function including source code and parameters."); }
+	virtual FString GetCategory() const override { return TEXT("ToolsetRegistry"); }
+
+	virtual TArray<FECACommandParam> GetParameters() const override
+	{
+		return {
+			{ TEXT("function_name"), TEXT("string"), TEXT("Name of the registered function"), true }
+		};
+	}
+
+	virtual FECACommandResult Execute(const TSharedPtr<FJsonObject>& Params) override;
+};
+
+FECACommandResult FECACommand_GetRegisteredFunctionInfo::Execute(const TSharedPtr<FJsonObject>& Params)
+{
+	FString FunctionName;
+	if (!GetStringParam(Params, TEXT("function_name"), FunctionName))
+	{
+		return FECACommandResult::Error(TEXT("Missing required parameter: function_name"));
+	}
+
+	const ToolsetRegistryCache::FRegisteredFunctionInfo* FuncInfo = ToolsetRegistryCache::Find(FunctionName);
+	if (!FuncInfo)
+	{
+		return FECACommandResult::Error(FString::Printf(
+			TEXT("Function '%s' not found in toolset registry cache."), *FunctionName));
+	}
+
+	TSharedPtr<FJsonObject> Result = MakeResult();
+	Result->SetStringField(TEXT("name"), FuncInfo->Name);
+	Result->SetStringField(TEXT("description"), FuncInfo->Description);
+	Result->SetStringField(TEXT("category"), FuncInfo->Category);
+	Result->SetStringField(TEXT("scope"), FuncInfo->Scope);
+	Result->SetStringField(TEXT("return_type"), FuncInfo->ReturnType);
+
+	if (!FuncInfo->SourceCode.IsEmpty())
+	{
+		Result->SetStringField(TEXT("source_code"), FuncInfo->SourceCode);
+	}
+
+	TArray<TSharedPtr<FJsonValue>> ParamsArray;
+	for (const TSharedPtr<FJsonObject>& ParamObj : FuncInfo->Parameters)
+	{
+		ParamsArray.Add(MakeShared<FJsonValueObject>(ParamObj));
+	}
+	Result->SetArrayField(TEXT("parameters"), ParamsArray);
+
+	return FECACommandResult::Success(Result);
+}
+
+//------------------------------------------------------------------------------
+// register_function
+//
+// Queues a function registration request. ECA processes this by writing
+// the function to .eca/functions/ and loading it into the Python registry.
+// The bridge cache is then updated via store_registered_functions.
+//------------------------------------------------------------------------------
+
+class FECACommand_RegisterFunction : public IECACommand
+{
+public:
+	virtual FString GetName() const override { return TEXT("register_function"); }
+	virtual FString GetDescription() const override { return TEXT("Request registration of a new Python function. ECA will write and load the function."); }
+	virtual FString GetCategory() const override { return TEXT("ToolsetRegistry"); }
+
+	virtual TArray<FECACommandParam> GetParameters() const override
+	{
+		return {
+			{ TEXT("function_code"), TEXT("string"), TEXT("Python source code for the function"), true },
+			{ TEXT("function_name"), TEXT("string"), TEXT("Name for the function"), true },
+			{ TEXT("scope"), TEXT("string"), TEXT("Scope: 'project' or 'global'"), false, TEXT("project") },
+			{ TEXT("category"), TEXT("string"), TEXT("Category for organization"), false, TEXT("User-Defined") }
+		};
+	}
+
+	virtual FECACommandResult Execute(const TSharedPtr<FJsonObject>& Params) override;
+};
+
+FECACommandResult FECACommand_RegisterFunction::Execute(const TSharedPtr<FJsonObject>& Params)
+{
+	FString FunctionCode;
+	if (!GetStringParam(Params, TEXT("function_code"), FunctionCode))
+	{
+		return FECACommandResult::Error(TEXT("Missing required parameter: function_code"));
+	}
+
+	FString FunctionName;
+	if (!GetStringParam(Params, TEXT("function_name"), FunctionName))
+	{
+		return FECACommandResult::Error(TEXT("Missing required parameter: function_name"));
+	}
+
+	FString Scope = TEXT("project");
+	GetStringParam(Params, TEXT("scope"), Scope, false);
+
+	FString Category = TEXT("User-Defined");
+	GetStringParam(Params, TEXT("category"), Category, false);
+
+	// Build the result with the registration request info
+	// ECA will execute this via its own register_function tool
+	TSharedPtr<FJsonObject> Result = MakeResult();
+	Result->SetStringField(TEXT("function_name"), FunctionName);
+	Result->SetStringField(TEXT("function_code"), FunctionCode);
+	Result->SetStringField(TEXT("scope"), Scope);
+	Result->SetStringField(TEXT("category"), Category);
+	Result->SetStringField(TEXT("status"), TEXT("registration_requested"));
+	Result->SetStringField(TEXT("instruction"),
+		FString::Printf(TEXT("Register this function via meta_tool: tools.register_function(function_code=%s, function_name='%s', scope='%s', category='%s')"),
+			*FunctionCode, *FunctionName, *Scope, *Category));
+
+	return FECACommandResult::Success(Result);
+}
+
+//------------------------------------------------------------------------------
+// unregister_function
+//------------------------------------------------------------------------------
+
+class FECACommand_UnregisterFunction : public IECACommand
+{
+public:
+	virtual FString GetName() const override { return TEXT("unregister_function"); }
+	virtual FString GetDescription() const override { return TEXT("Request unregistration of a registered Python function."); }
+	virtual FString GetCategory() const override { return TEXT("ToolsetRegistry"); }
+
+	virtual TArray<FECACommandParam> GetParameters() const override
+	{
+		return {
+			{ TEXT("function_name"), TEXT("string"), TEXT("Name of the function to unregister"), true },
+			{ TEXT("scope"), TEXT("string"), TEXT("Scope to search in"), false }
+		};
+	}
+
+	virtual FECACommandResult Execute(const TSharedPtr<FJsonObject>& Params) override;
+};
+
+FECACommandResult FECACommand_UnregisterFunction::Execute(const TSharedPtr<FJsonObject>& Params)
+{
+	FString FunctionName;
+	if (!GetStringParam(Params, TEXT("function_name"), FunctionName))
+	{
+		return FECACommandResult::Error(TEXT("Missing required parameter: function_name"));
+	}
+
+	FString Scope;
+	GetStringParam(Params, TEXT("scope"), Scope, false);
+
+	TSharedPtr<FJsonObject> Result = MakeResult();
+	Result->SetStringField(TEXT("function_name"), FunctionName);
+	if (!Scope.IsEmpty())
+	{
+		Result->SetStringField(TEXT("scope"), Scope);
+	}
+	Result->SetStringField(TEXT("status"), TEXT("unregistration_requested"));
+	Result->SetStringField(TEXT("instruction"),
+		FString::Printf(TEXT("Unregister this function via meta_tool: tools.unregister_function(name='%s'%s)"),
+			*FunctionName,
+			Scope.IsEmpty() ? TEXT("") : *FString::Printf(TEXT(", scope='%s'"), *Scope)));
+
+	return FECACommandResult::Success(Result);
+}
+
+//------------------------------------------------------------------------------
+// get_all_tools
+//------------------------------------------------------------------------------
+
+class FECACommand_GetAllTools : public IECACommand
+{
+public:
+	virtual FString GetName() const override { return TEXT("get_all_tools"); }
+	virtual FString GetDescription() const override { return TEXT("Get a combined list of all available tools: ECABridge commands and registered Python functions."); }
+	virtual FString GetCategory() const override { return TEXT("ToolsetRegistry"); }
+
+	virtual TArray<FECACommandParam> GetParameters() const override
+	{
+		return {
+			{ TEXT("include_bridge_commands"), TEXT("boolean"), TEXT("Include ECABridge commands"), false, TEXT("true") },
+			{ TEXT("include_registered_functions"), TEXT("boolean"), TEXT("Include registered Python functions"), false, TEXT("true") },
+			{ TEXT("category_filter"), TEXT("string"), TEXT("Filter by category"), false }
+		};
+	}
+
+	virtual FECACommandResult Execute(const TSharedPtr<FJsonObject>& Params) override;
+};
+
+FECACommandResult FECACommand_GetAllTools::Execute(const TSharedPtr<FJsonObject>& Params)
+{
+	bool bIncludeBridgeCommands = true;
+	GetBoolParam(Params, TEXT("include_bridge_commands"), bIncludeBridgeCommands, false);
+
+	bool bIncludeRegisteredFunctions = true;
+	GetBoolParam(Params, TEXT("include_registered_functions"), bIncludeRegisteredFunctions, false);
+
+	FString CategoryFilter;
+	GetStringParam(Params, TEXT("category_filter"), CategoryFilter, false);
+
+	TSharedPtr<FJsonObject> Result = MakeResult();
+	TArray<TSharedPtr<FJsonValue>> AllTools;
+	int32 BridgeCount = 0;
+	int32 RegisteredCount = 0;
+
+	// Add ECABridge commands
+	if (bIncludeBridgeCommands)
+	{
+		for (const TSharedPtr<IECACommand>& Command : FECACommandRegistry::Get().GetAllCommands())
+		{
+			if (!CategoryFilter.IsEmpty() && Command->GetCategory() != CategoryFilter)
+			{
+				continue;
+			}
+
+			TSharedPtr<FJsonObject> ToolObj = MakeShared<FJsonObject>();
+			ToolObj->SetStringField(TEXT("name"), Command->GetName());
+			ToolObj->SetStringField(TEXT("description"), Command->GetDescription());
+			ToolObj->SetStringField(TEXT("category"), Command->GetCategory());
+			ToolObj->SetStringField(TEXT("source"), TEXT("eca_bridge"));
+
+			TArray<TSharedPtr<FJsonValue>> ParamsArray;
+			for (const FECACommandParam& Param : Command->GetParameters())
+			{
+				TSharedPtr<FJsonObject> ParamObj = MakeShared<FJsonObject>();
+				ParamObj->SetStringField(TEXT("name"), Param.Name);
+				ParamObj->SetStringField(TEXT("type"), Param.Type);
+				ParamObj->SetStringField(TEXT("description"), Param.Description);
+				ParamObj->SetBoolField(TEXT("required"), Param.bRequired);
+				if (!Param.DefaultValue.IsEmpty())
+				{
+					ParamObj->SetStringField(TEXT("default"), Param.DefaultValue);
+				}
+				ParamsArray.Add(MakeShared<FJsonValueObject>(ParamObj));
+			}
+			ToolObj->SetArrayField(TEXT("parameters"), ParamsArray);
+
+			AllTools.Add(MakeShared<FJsonValueObject>(ToolObj));
+			BridgeCount++;
+		}
+	}
+
+	// Add registered functions from cache
+	if (bIncludeRegisteredFunctions)
+	{
+		TArray<ToolsetRegistryCache::FRegisteredFunctionInfo> Functions =
+			ToolsetRegistryCache::GetAll(TEXT(""), CategoryFilter);
+
+		for (const auto& Info : Functions)
+		{
+			TSharedPtr<FJsonObject> ToolObj = MakeShared<FJsonObject>();
+			ToolObj->SetStringField(TEXT("name"), Info.Name);
+			ToolObj->SetStringField(TEXT("description"), Info.Description);
+			ToolObj->SetStringField(TEXT("category"), Info.Category);
+			ToolObj->SetStringField(TEXT("source"), TEXT("toolset_registry"));
+			ToolObj->SetStringField(TEXT("scope"), Info.Scope);
+			ToolObj->SetStringField(TEXT("return_type"), Info.ReturnType);
+
+			TArray<TSharedPtr<FJsonValue>> ParamsArray;
+			for (const TSharedPtr<FJsonObject>& ParamObj : Info.Parameters)
+			{
+				ParamsArray.Add(MakeShared<FJsonValueObject>(ParamObj));
+			}
+			ToolObj->SetArrayField(TEXT("parameters"), ParamsArray);
+
+			AllTools.Add(MakeShared<FJsonValueObject>(ToolObj));
+			RegisteredCount++;
+		}
+	}
+
+	Result->SetArrayField(TEXT("tools"), AllTools);
+	Result->SetNumberField(TEXT("total_count"), AllTools.Num());
+	Result->SetNumberField(TEXT("bridge_command_count"), BridgeCount);
+	Result->SetNumberField(TEXT("registered_function_count"), RegisteredCount);
+
+	return FECACommandResult::Success(Result);
+}
+
+//------------------------------------------------------------------------------
+// Register all commands
+//------------------------------------------------------------------------------
+
+REGISTER_ECA_COMMAND(FECACommand_StoreRegisteredFunctions)
+REGISTER_ECA_COMMAND(FECACommand_ListRegisteredFunctions)
+REGISTER_ECA_COMMAND(FECACommand_CallRegisteredFunction)
+REGISTER_ECA_COMMAND(FECACommand_GetRegisteredFunctionInfo)
+REGISTER_ECA_COMMAND(FECACommand_RegisterFunction)
+REGISTER_ECA_COMMAND(FECACommand_UnregisterFunction)
+REGISTER_ECA_COMMAND(FECACommand_GetAllTools)
\ No newline at end of file
diff --git a/Engine/Plugins/Experimental/ECABridge/Source/ECABridge/Private/ECABridgeModule.cpp b/Engine/Plugins/Experimental/ECABridge/Source/ECABridge/Private/ECABridgeModule.cpp
--- a/Engine/Plugins/Experimental/ECABridge/Source/ECABridge/Private/ECABridgeModule.cpp
+++ b/Engine/Plugins/Experimental/ECABridge/Source/ECABridge/Private/ECABridgeModule.cpp
@@ -3,9 +3,13 @@
 #include "ECABridgeModule.h"
 #include "ECABridge.h"
 #include "ECABridgeSettings.h"
+#if WITH_TOOLSET_REGISTRY
 #include "ECABridgeToolset.h"
+#endif
 #include "ECAEventQueue.h"
+#if WITH_TOOLSET_REGISTRY
 #include "ToolsetRegistry/ToolsetRegistrySubsystem.h"
+#endif
 #include "HAL/IConsoleManager.h"
 #include "Editor.h"
 #include "Selection.h"
@@ -43,6 +47,7 @@ void FECABridgeModule::StartupModule()
 	// Register editor event handlers
 	RegisterEditorEventHandlers();
 	
+#if WITH_TOOLSET_REGISTRY
 	// Register ECABridge commands with the ToolsetRegistry (Direction 1)
 	// Deferred until GEditor is available since ToolsetRegistrySubsystem is an editor subsystem
 	if (GEditor)
@@ -50,14 +55,17 @@ void FECABridgeModule::StartupModule()
 		RegisterWithToolsetRegistry();
 	}
 	// else: handled in OnEditorInitialized/RegisterGEditorDependentHandlers
+#endif
 	
 	UE_LOG(LogTemp, Log, TEXT("[ECABridge] Module ready"));
 }
 
 void FECABridgeModule::ShutdownModule()
 {
+#if WITH_TOOLSET_REGISTRY
 	// Unregister from ToolsetRegistry
 	UnregisterFromToolsetRegistry();
+#endif
 	
 	// Unregister editor event handlers
 	UnregisterEditorEventHandlers();
@@ -186,11 +194,13 @@ void FECABridgeModule::RegisterGEditorDependentHandlers()
 	CameraMovedHandle = FEditorDelegates::OnEditorCameraMoved.AddRaw(this, &FECABridgeModule::OnEditorCameraMoved);
 	UE_LOG(LogTemp, Log, TEXT("[ECABridge] Registered camera movement handler"));
 	
+#if WITH_TOOLSET_REGISTRY
 	// Register with ToolsetRegistry now that GEditor is available (if not already done)
 	if (!ECABridgeToolset.IsValid())
 	{
 		RegisterWithToolsetRegistry();
 	}
+#endif
 	
 	UE_LOG(LogTemp, Log, TEXT("[ECABridge] All GEditor-dependent handlers registered"));
 }
@@ -292,6 +302,7 @@ void FECABridgeModule::UnregisterEditorEventHandlers()
 	UE_LOG(LogTemp, Log, TEXT("[ECABridge] All editor event handlers unregistered"));
 }
 
+#if WITH_TOOLSET_REGISTRY
 void FECABridgeModule::RegisterWithToolsetRegistry()
 {
 	if (!GEditor)
@@ -343,6 +354,7 @@ void FECABridgeModule::UnregisterFromToolsetRegistry()
 		ToolsetRegistryChangedHandle.Reset();
 	}
 }
+#endif // WITH_TOOLSET_REGISTRY
 
 void FECABridgeModule::OnEditorSelectionChanged(UObject* SelectionObject)
 {
diff --git a/Engine/Plugins/Experimental/ECABridge/Source/ECABridge/Private/ECABridgeToolset.cpp b/Engine/Plugins/Experimental/ECABridge/Source/ECABridge/Private/ECABridgeToolset.cpp
--- a/Engine/Plugins/Experimental/ECABridge/Source/ECABridge/Private/ECABridgeToolset.cpp
+++ b/Engine/Plugins/Experimental/ECABridge/Source/ECABridge/Private/ECABridgeToolset.cpp
@@ -1,6 +1,8 @@
 // Copyright Epic Games, Inc. All Rights Reserved.
 
 #include "ECABridgeToolset.h"
+
+#if WITH_TOOLSET_REGISTRY
 #include "Commands/ECACommand.h"
 #include "Dom/JsonObject.h"
 #include "Dom/JsonValue.h"
@@ -182,3 +184,4 @@ FString FECABridgeToolset::GetToolsetJsonSchema() const
 	FJsonSerializer::Serialize(ToolsetEnvelope.ToSharedRef(), Writer);
 	return OutputString;
 }
+#endif // WITH_TOOLSET_REGISTRY
diff --git a/Engine/Plugins/Experimental/ECABridge/Source/ECABridge/Private/ECAMCPServer.cpp b/Engine/Plugins/Experimental/ECABridge/Source/ECABridge/Private/ECAMCPServer.cpp
--- a/Engine/Plugins/Experimental/ECABridge/Source/ECABridge/Private/ECAMCPServer.cpp
+++ b/Engine/Plugins/Experimental/ECABridge/Source/ECABridge/Private/ECAMCPServer.cpp
@@ -12,8 +12,10 @@
 #include "Serialization/JsonWriter.h"
 #include "Serialization/JsonSerializer.h"
 #include "Editor.h"
+#if WITH_TOOLSET_REGISTRY
 #include "ToolsetRegistry/ToolsetRegistry.h"
 #include "ToolsetRegistry/ToolsetRegistrySubsystem.h"
+#endif
 
 FECAMCPServer::FECAMCPServer(UECABridge* InBridge)
 	: Bridge(InBridge)
@@ -418,12 +420,14 @@ TSharedPtr<FJsonObject> FECAMCPServer::HandleToolsCall(const TSharedPtr<FJsonObj
 		return Result;
 	}
 	
+#if WITH_TOOLSET_REGISTRY
 	// Direction 2: Try ToolsetRegistry if not found in ECABridge commands
 	TSharedPtr<FJsonObject> ToolsetResult;
 	if (TryExecuteToolsetRegistryTool(ToolName, Arguments, ToolsetResult))
 	{
 		return ToolsetResult;
 	}
+#endif
 	
 	// Tool not found anywhere
 	TSharedPtr<FJsonObject> ErrorResult = MakeShared<FJsonObject>();
@@ -518,9 +522,11 @@ TArray<TSharedPtr<FJsonValue>> FECAMCPServer::BuildToolDefinitions()
 		ExistingNames.Add(Command->GetName().ToLower());
 	}
 
+#if WITH_TOOLSET_REGISTRY
 	// Direction 2: Also include tools from ToolsetRegistry (excluding our own ECABridge toolset
 	// to avoid duplication — we already added those commands above)
 	AppendToolsetRegistryTools(Tools, ExistingNames);
+#endif
 
 	return Tools;
 }
@@ -562,9 +568,9 @@ TSharedPtr<FJsonObject> FECAMCPServer::CreateJsonRpcError(int32 Code, const FStr
 	return Response;
 }
 
+#if WITH_TOOLSET_REGISTRY
 void FECAMCPServer::AppendToolsetRegistryTools(TArray<TSharedPtr<FJsonValue>>& Tools, TSet<FString>& ExistingNames)
 {
-#if WITH_EDITOR
 	UToolsetRegistrySubsystem* ToolsetRegistrySubsystem = 
 		GEditor ? GEditor->GetEditorSubsystem<UToolsetRegistrySubsystem>() : nullptr;
 	if (!ToolsetRegistrySubsystem)
@@ -655,13 +661,11 @@ void FECAMCPServer::AppendToolsetRegistryTools(TArray<TSharedPtr<FJsonValue>>& T
 	{
 		UE_LOG(LogTemp, Log, TEXT("[ECABridge] Added %d tools from ToolsetRegistry to MCP tool list"), AddedCount);
 	}
-#endif // WITH_EDITOR
 }
 
 bool FECAMCPServer::TryExecuteToolsetRegistryTool(const FString& ToolName, const TSharedPtr<FJsonObject>& Arguments,
 	TSharedPtr<FJsonObject>& OutResult)
 {
-#if WITH_EDITOR
 	UToolsetRegistrySubsystem* ToolsetRegistrySubsystem = 
 		GEditor ? GEditor->GetEditorSubsystem<UToolsetRegistrySubsystem>() : nullptr;
 	if (!ToolsetRegistrySubsystem)
@@ -739,10 +743,8 @@ bool FECAMCPServer::TryExecuteToolsetRegistryTool(const FString& ToolName, const
 	Content.Add(MakeShared<FJsonValueObject>(TextContent));
 	OutResult->SetArrayField(TEXT("content"), Content);
 	return true;
-#else
-	return false;
-#endif // WITH_EDITOR
 }
+#endif // WITH_TOOLSET_REGISTRY
 
 FString FECAMCPServer::SerializeJson(const TSharedPtr<FJsonObject>& JsonObject)
 {
diff --git a/Engine/Plugins/Experimental/ECABridge/Source/ECABridge/Public/BlueprintLisp.h b/Engine/Plugins/Experimental/ECABridge/Source/ECABridge/Public/BlueprintLisp.h
--- a/Engine/Plugins/Experimental/ECABridge/Source/ECABridge/Public/BlueprintLisp.h
+++ b/Engine/Plugins/Experimental/ECABridge/Source/ECABridge/Public/BlueprintLisp.h
@@ -34,7 +34,7 @@
  *   (delay seconds)            - Delay node
  *   (cast Type var body)       - Cast with success exec
  *   (vec x y z)                - Vector literal
- *   (rot p y r)                - Rotator literal
+ *   (rot r p y)                - Rotator literal (Roll Pitch Yaw)
  */
 
 //------------------------------------------------------------------------------
diff --git a/Engine/Plugins/Experimental/ECABridge/Source/ECABridge/Public/ECABridgeModule.h b/Engine/Plugins/Experimental/ECABridge/Source/ECABridge/Public/ECABridgeModule.h
--- a/Engine/Plugins/Experimental/ECABridge/Source/ECABridge/Public/ECABridgeModule.h
+++ b/Engine/Plugins/Experimental/ECABridge/Source/ECABridge/Public/ECABridgeModule.h
@@ -6,14 +6,18 @@
 #include "Modules/ModuleManager.h"
 #include "Engine/EngineTypes.h"  // For FTimerHandle
 #include "Editor/UnrealEdTypes.h"  // For ELevelViewportType
+#if WITH_TOOLSET_REGISTRY
 #include "ToolsetRegistry/Toolset.h"
 #include "ToolsetRegistry/ToolsetRegistry.h"
+#endif
 
 class UECABridge;
 class UFactory;
 struct FAssetData;
 class AActor;
+#if WITH_TOOLSET_REGISTRY
 class FECABridgeToolset;
+#endif
 
 /**
  * ECA Bridge Module
@@ -53,11 +57,13 @@ private:
 	/** Console commands for debugging */
 	void RegisterConsoleCommands();
 	
+#if WITH_TOOLSET_REGISTRY
 	/** Register ECABridge commands as a ToolsetRegistry toolset */
 	void RegisterWithToolsetRegistry();
 	
 	/** Unregister from ToolsetRegistry */
 	void UnregisterFromToolsetRegistry();
+#endif
 	
 	/** Register editor event handlers for pending events */
 	void RegisterEditorEventHandlers();
@@ -150,6 +156,7 @@ private:
 	/** Debounce time in seconds - wait this long after camera stops before sending event */
 	static constexpr float CameraStoppedDebounceTime = 0.5f;
 	
+#if WITH_TOOLSET_REGISTRY
 	//--------------------------------------------------------------------------
 	// ToolsetRegistry Integration
 	//--------------------------------------------------------------------------
@@ -159,4 +166,5 @@ private:
 	
 	/** Handle for ToolsetRegistry change notifications */
 	FDelegateHandle ToolsetRegistryChangedHandle;
+#endif
 };
diff --git a/Engine/Plugins/Experimental/ECABridge/Source/ECABridge/Public/ECABridgeToolset.h b/Engine/Plugins/Experimental/ECABridge/Source/ECABridge/Public/ECABridgeToolset.h
--- a/Engine/Plugins/Experimental/ECABridge/Source/ECABridge/Public/ECABridgeToolset.h
+++ b/Engine/Plugins/Experimental/ECABridge/Source/ECABridge/Public/ECABridgeToolset.h
@@ -3,6 +3,8 @@
 #pragma once
 
 #include "CoreMinimal.h"
+
+#if WITH_TOOLSET_REGISTRY
 #include "ToolsetRegistry/Toolset.h"
 
 /**
@@ -31,3 +33,4 @@ public:
 	virtual FString GetToolsetVersion() const override;
 	//~ End FToolset interface
 };
+#endif // WITH_TOOLSET_REGISTRY
diff --git a/Engine/Plugins/Experimental/ECABridge/Source/ECABridge/Public/ECAMCPServer.h b/Engine/Plugins/Experimental/ECABridge/Source/ECABridge/Public/ECAMCPServer.h
--- a/Engine/Plugins/Experimental/ECABridge/Source/ECABridge/Public/ECAMCPServer.h
+++ b/Engine/Plugins/Experimental/ECABridge/Source/ECABridge/Public/ECAMCPServer.h
@@ -64,15 +64,17 @@ private:
 	/** Handle tools/call request */
 	TSharedPtr<FJsonObject> HandleToolsCall(const TSharedPtr<FJsonObject>& Params);
 	
-	/** Build MCP tool definitions from registered commands and ToolsetRegistry */
+	/** Build MCP tool definitions from registered commands (and ToolsetRegistry if available) */
 	TArray<TSharedPtr<FJsonValue>> BuildToolDefinitions();
-	
+
+#if WITH_TOOLSET_REGISTRY
 	/** Build tool definitions from ToolsetRegistry toolsets (Direction 2) */
 	void AppendToolsetRegistryTools(TArray<TSharedPtr<FJsonValue>>& Tools, TSet<FString>& ExistingNames);
 	
 	/** Try to execute a tool via ToolsetRegistry if not found in ECABridge commands */
 	bool TryExecuteToolsetRegistryTool(const FString& ToolName, const TSharedPtr<FJsonObject>& Arguments,
 		TSharedPtr<FJsonObject>& OutResult);
+#endif
 	
 	/** Create JSON-RPC response wrapper */
 	TSharedPtr<FJsonObject> CreateJsonRpcResponse(const TSharedPtr<FJsonObject>& Result, const TSharedPtr<FJsonValue>& RequestId);
