[ECABridge] Blueprint node graph: create_widget node type in batch_edit_blueprint_nodes

Adds "create_widget" as a new node type in the batch_edit_blueprint_nodes command.

Implementation:
- FindObject<UClass> + LoadObject fallback for K2Node_CreateWidget (avoids private UMGEditor header)
- Correct lifecycle order: AllocateDefaultPins ??? set ClassPin->DefaultObject ??? PostPlacedNewNode
(PostPlacedNewNode calls CreatePinsForClass which needs the Class pin set first to create expose-on-spawn property pins)
- Does NOT call Graph->AddNode or AllocateDefaultPins itself ??? common finalization block handles those for all node types

#rnx

#Synced-CL 52128568

diff --git a/Engine/Plugins/Experimental/ECABridge/Source/ECABridge/Private/Commands/ECABlueprintNodeCommands.cpp b/Engine/Plugins/Experimental/ECABridge/Source/ECABridge/Private/Commands/ECABlueprintNodeCommands.cpp
--- a/Engine/Plugins/Experimental/ECABridge/Source/ECABridge/Private/Commands/ECABlueprintNodeCommands.cpp
+++ b/Engine/Plugins/Experimental/ECABridge/Source/ECABridge/Private/Commands/ECABlueprintNodeCommands.cpp
@@ -1698,9 +1698,59 @@ FECACommandResult FECACommand_BatchEditBlueprintNodes::Execute(const TSharedPtr<
 			MacroNode->SetMacroGraph(MacroGraph);
 			CreatedNode = MacroNode;
 		}
+		else if (NodeType == TEXT("create_widget"))
+		{
+			// Create a UK2Node_CreateWidget node — avoids including the private UMGEditor header.
+			FString WidgetClass;
+			NodeDef->TryGetStringField(TEXT("widget_class"), WidgetClass);
+
+			UClass* CreateWidgetNodeClass = FindObject<UClass>(nullptr, TEXT("/Script/UMGEditor.K2Node_CreateWidget"));
+			if (!CreateWidgetNodeClass)
+			{
+				CreateWidgetNodeClass = LoadObject<UClass>(nullptr, TEXT("/Script/UMGEditor.K2Node_CreateWidget"));
+			}
+			if (!CreateWidgetNodeClass)
+			{
+				Errors.Add(FString::Printf(TEXT("Node '%s': K2Node_CreateWidget class not found (UMGEditor not loaded?)"), *TempId));
+				continue;
+			}
+
+			UK2Node* CreateWidgetNode = NewObject<UK2Node>(Graph, CreateWidgetNodeClass);
+			CreateWidgetNode->CreateNewGuid();
+
+			// Correct lifecycle order for K2Node_ConstructObjectFromClass subclasses:
+			// 1. AllocateDefaultPins — creates the Class pin (required before step 2)
+			// 2. Set Class pin DefaultObject — so PostPlacedNewNode sees the class
+			// 3. PostPlacedNewNode — calls CreatePinsForClass to add expose-on-spawn pins
+			CreateWidgetNode->AllocateDefaultPins();
+
+			if (!WidgetClass.IsEmpty())
+			{
+				UClass* WidgetClassObj = LoadObject<UClass>(nullptr, *WidgetClass);
+				if (!WidgetClassObj)
+				{
+					WidgetClassObj = LoadObject<UClass>(nullptr, *FString::Printf(TEXT("%s_C"), *WidgetClass));
+				}
+				if (WidgetClassObj)
+				{
+					if (UEdGraphPin* ClassPin = CreateWidgetNode->FindPin(TEXT("Class")))
+					{
+						ClassPin->DefaultObject = WidgetClassObj;
+					}
+				}
+			}
+
+			CreateWidgetNode->PostPlacedNewNode();
+
+			// Do NOT call Graph->AddNode or set position here —
+			// the common finalization block handles that for every CreatedNode.
+			// AllocateDefaultPins was already called above (step 1 of the lifecycle)
+			// and PostPlacedNewNode has run; the finalization block skips ADP when Pins.Num() > 0.
+			CreatedNode = CreateWidgetNode;
+		}
 		else
 		{
-			Errors.Add(FString::Printf(TEXT("Node '%s': Unknown type '%s'. Valid types: event, custom_event, function, variable_get, variable_set, self, component, input_action, cast, macro"), *TempId, *NodeType));
+			Errors.Add(FString::Printf(TEXT("Node '%s': Unknown type '%s'. Valid types: event, custom_event, function, variable_get, variable_set, self, component, input_action, cast, macro, create_widget"), *TempId, *NodeType));
 			continue;
 		}
 		
@@ -1709,7 +1759,12 @@ FECACommandResult FECACommand_BatchEditBlueprintNodes::Execute(const TSharedPtr<
 			CreatedNode->NodePosX = Position.X;
 			CreatedNode->NodePosY = Position.Y;
 			Graph->AddNode(CreatedNode, false, false);
-			CreatedNode->AllocateDefaultPins();
+			// Skip AllocateDefaultPins for nodes that already called it (e.g. create_widget
+			// which must call it early to set up the Class pin before PostPlacedNewNode).
+			if (CreatedNode->Pins.Num() == 0)
+			{
+				CreatedNode->AllocateDefaultPins();
+			}
 			EnsureBPNodeHasValidGuid(CreatedNode);
 			
 			TempIdToNode.Add(TempId, CreatedNode);
