[ECABridge] MVVM multi-source bindings, bind_visibility improvements, set_viewmodel_binding_argument

Modified commands:
- add_viewmodel_binding: accepts source_properties[] + source_viewmodels[] arrays for cross-VM multi-input conversion functions. source_properties takes precedence over source_property when both provided. Per-source VM override with null guard (skips source if property not found on override VM). NewBinding reference invalidation documented. IsValidConversionFunction pre-check uses first source. ParseBindingMode logs UE_LOG warning on unknown string.
- bind_visibility_to_viewmodel: viewmodel_name param targets specific VM; searches all VMs if not specified (TArrayView<const...> stable pointer documented). Handles UFunction sources (FieldNotify getters) via GetReturnProperty() for ExpectedType. conversion_function result fields now conditional on ConvFunc != null.

New command:
- set_viewmodel_binding_argument: patches a specific argument pin of an existing binding's conversion function by binding_index + arg_index. Resolves conv function by name via ResolveConversionFunction + TObjectIterator K2Node fallback. Enables cross-VM wiring for multi-input conversion functions.

#rnx
#Synced-CL 52128568

diff --git a/Engine/Plugins/Experimental/ECABridge/Source/ECABridge/Private/Commands/ECAMVVMCommands.cpp b/Engine/Plugins/Experimental/ECABridge/Source/ECABridge/Private/Commands/ECAMVVMCommands.cpp
--- a/Engine/Plugins/Experimental/ECABridge/Source/ECABridge/Private/Commands/ECAMVVMCommands.cpp
+++ b/Engine/Plugins/Experimental/ECABridge/Source/ECABridge/Private/Commands/ECAMVVMCommands.cpp
@@ -9,6 +9,8 @@
 #include "Components/TextBlock.h"
 #include "Components/Widget.h"
 #include "EdGraph/EdGraphPin.h"
+#include "K2Node.h"
+#include "Types/MVVMConversionFunctionValue.h"
 #include "Editor.h"
 #include "Kismet/BlueprintFunctionLibrary.h"
 #include "Kismet2/BlueprintEditorUtils.h"
@@ -36,6 +38,7 @@ REGISTER_ECA_COMMAND(FECACommand_GetViewmodelBindings)
 REGISTER_ECA_COMMAND(FECACommand_RemoveViewmodelBinding)
 REGISTER_ECA_COMMAND(FECACommand_GetWidgetMVVMProperties)
 REGISTER_ECA_COMMAND(FECACommand_AddViewmodelBinding)
+REGISTER_ECA_COMMAND(FECACommand_SetViewmodelBindingArgument)
 
 //------------------------------------------------------------------------------
 // Helper Functions
@@ -192,7 +195,11 @@ static EMVVMBindingMode ParseBindingMode(const FString& ModeStr)
 	{
 		return EMVVMBindingMode::OneTimeToDestination;
 	}
-	// Default
+	// Default — warn on unknown values rather than silently falling through
+	if (!ModeStr.IsEmpty())
+	{
+		UE_LOG(LogTemp, Warning, TEXT("[ECABridge] ParseBindingMode: unknown binding mode '%s', defaulting to OneWayToDestination"), *ModeStr);
+	}
 	return EMVVMBindingMode::OneWayToDestination;
 }
 
@@ -417,8 +424,10 @@ FECACommandResult FECACommand_BindVisibilityToViewmodel::Execute(const TSharedPt
 
 	FString TrueVisibility = TEXT("Visible");
 	FString FalseVisibility = TEXT("Collapsed");
+	FString VisVMName;
 	GetStringParam(Params, TEXT("true_visibility"), TrueVisibility, false);
 	GetStringParam(Params, TEXT("false_visibility"), FalseVisibility, false);
+	GetStringParam(Params, TEXT("viewmodel_name"), VisVMName, false);
 
 	// Validate visibility enum values
 	TArray<FString> ValidVisibilities = {
@@ -458,16 +467,44 @@ FECACommandResult FECACommand_BindVisibilityToViewmodel::Execute(const TSharedPt
 		return FECACommandResult::Error(FString::Printf(TEXT("Widget '%s' not found in the widget tree."), *WidgetName));
 	}
 
-	const FMVVMBlueprintViewModelContext& VMContext = View->GetViewModels()[0];
-	const UClass* VMClass = VMContext.GetViewModelClass();
-
-	FProperty* VMProp = FindPropertyOnClass(VMClass, VMProperty);
-	if (!VMProp)
+	// Find ViewModel — by name if specified, otherwise search all VMs for property or function
+	const FMVVMBlueprintViewModelContext* VMContextPtr = nullptr;
+	FProperty* VMProp = nullptr;
+	UFunction* VMFunc = nullptr;
+	if (!VisVMName.IsEmpty())
+	{
+		VMContextPtr = View->FindViewModel(FName(*VisVMName));
+		if (!VMContextPtr)
+		{
+			return FECACommandResult::Error(FString::Printf(TEXT("ViewModel '%s' not found."), *VisVMName));
+		}
+		const UClass* C = VMContextPtr->GetViewModelClass();
+		VMProp = FindPropertyOnClass(C, VMProperty);
+		if (!VMProp && C) VMFunc = C->FindFunctionByName(*VMProperty);
+	}
+	else
+	{
+		// GetViewModels() returns TArrayView<const ...> — a stable view into the internal array.
+		// Storing &Ctx is safe as long as we don't mutate the view collection during this loop.
+		for (const FMVVMBlueprintViewModelContext& Ctx : View->GetViewModels())
+		{
+			const UClass* C = Ctx.GetViewModelClass();
+			FProperty* Prop = FindPropertyOnClass(C, VMProperty);
+			if (Prop) { VMContextPtr = &Ctx; VMProp = Prop; break; }
+			UFunction* Func = C ? C->FindFunctionByName(*VMProperty) : nullptr;
+			if (Func) { VMContextPtr = &Ctx; VMFunc = Func; break; }
+		}
+	}
+	if (!VMContextPtr || (!VMProp && !VMFunc))
 	{
 		return FECACommandResult::Error(FString::Printf(
-			TEXT("ViewModel property not found: %s on %s. Check property name spelling."),
-			*VMProperty, *VMClass->GetName()));
+			TEXT("ViewModel property or function '%s' not found."), *VMProperty));
 	}
+	const FMVVMBlueprintViewModelContext& VMContext = *VMContextPtr;
+	const UClass* VMClass = VMContext.GetViewModelClass();
+	UE::MVVM::FMVVMConstFieldVariant VMField = VMProp
+		? UE::MVVM::FMVVMConstFieldVariant(VMProp)
+		: UE::MVVM::FMVVMConstFieldVariant(VMFunc);
 
 	// Create binding
 	FMVVMBlueprintViewBinding& NewBinding = MVVMSubsystem->AddBinding(WidgetBP);
@@ -502,14 +539,15 @@ FECACommandResult FECACommand_BindVisibilityToViewmodel::Execute(const TSharedPt
 		FMVVMBlueprintFunctionReference FuncRef(WidgetBP, ConvFunc);
 		MVVMSubsystem->SetSourceToDestinationConversionFunction(WidgetBP, NewBinding, FuncRef);
 
-		// Wire ViewModel bool to the bIsVisible parameter pin
+		// Wire ViewModel source to the bool parameter pin
 		UE::MVVM::FConversionFunctionValue ConvValue(ConvFunc);
-		TArray<FMVVMBlueprintPinId> BoolPins = MVVMSubsystem->GetConversionFunctionPins(WidgetBP, ConvValue, VMProp);
+		const FProperty* ExpectedType = VMProp ? VMProp : (VMFunc ? VMFunc->GetReturnProperty() : nullptr);
+		TArray<FMVVMBlueprintPinId> BoolPins = MVVMSubsystem->GetConversionFunctionPins(WidgetBP, ConvValue, ExpectedType);
 		if (BoolPins.Num() > 0)
 		{
 			FMVVMBlueprintPropertyPath VMSourcePath;
 			VMSourcePath.SetViewModelId(VMContext.GetViewModelId());
-			VMSourcePath.SetPropertyPath(WidgetBP, UE::MVVM::FMVVMConstFieldVariant(VMProp));
+			VMSourcePath.SetPropertyPath(WidgetBP, VMField);
 			MVVMSubsystem->SetPathForConversionFunctionArgument(WidgetBP, NewBinding, BoolPins[0], VMSourcePath, true);
 		}
 
@@ -555,16 +593,19 @@ FECACommandResult FECACommand_BindVisibilityToViewmodel::Execute(const TSharedPt
 		// Fallback: direct binding without conversion function
 		FMVVMBlueprintPropertyPath SourcePath;
 		SourcePath.SetViewModelId(VMContext.GetViewModelId());
-		SourcePath.SetPropertyPath(WidgetBP, UE::MVVM::FMVVMConstFieldVariant(VMProp));
+		SourcePath.SetPropertyPath(WidgetBP, VMField);
 		MVVMSubsystem->SetSourcePathForBinding(WidgetBP, NewBinding, SourcePath);
 	}
 
 	CompileAndMarkDirty(WidgetBP);
 
 	TSharedPtr<FJsonObject> Result = MakeResult();
-	Result->SetStringField(TEXT("conversion_function"), TEXT("Conv_BoolToSlateVisibility"));
-	Result->SetStringField(TEXT("true_visibility"), TrueVisibility);
-	Result->SetStringField(TEXT("false_visibility"), FalseVisibility);
+	if (ConvFunc)
+	{
+		Result->SetStringField(TEXT("conversion_function"), TEXT("Conv_BoolToSlateVisibility"));
+		Result->SetStringField(TEXT("true_visibility"), TrueVisibility);
+		Result->SetStringField(TEXT("false_visibility"), FalseVisibility);
+	}
 	return FECACommandResult::Success(Result);
 }
 
@@ -855,10 +896,62 @@ FECACommandResult FECACommand_AddViewmodelBinding::Execute(const TSharedPtr<FJso
 {
 	FString WidgetPath, SourceProperty, DestinationProperty;
 	if (!GetStringParam(Params, TEXT("widget_path"), WidgetPath) ||
-		!GetStringParam(Params, TEXT("source_property"), SourceProperty) ||
 		!GetStringParam(Params, TEXT("destination_property"), DestinationProperty))
 	{
-		return FECACommandResult::Error(TEXT("Missing required parameters: widget_path, source_property, destination_property"));
+		return FECACommandResult::Error(TEXT("Missing required parameters: widget_path, destination_property"));
+	}
+
+	// Accept either single source_property OR array source_properties (mutually exclusive).
+	// If source_properties array is provided, it takes precedence and source_property is ignored.
+	TArray<FString> SourceProperties;
+	const TArray<TSharedPtr<FJsonValue>>* SourcePropertiesArray = nullptr;
+	if (GetArrayParam(Params, TEXT("source_properties"), SourcePropertiesArray) && SourcePropertiesArray)
+	{
+		for (const TSharedPtr<FJsonValue>& Val : *SourcePropertiesArray)
+		{
+			FString PropName;
+			if (Val->TryGetString(PropName) && !PropName.IsEmpty())
+			{
+				SourceProperties.Add(PropName);
+			}
+		}
+	}
+	if (SourceProperties.IsEmpty())
+	{
+		GetStringParam(Params, TEXT("source_property"), SourceProperty, false);
+		if (!SourceProperty.IsEmpty())
+		{
+			SourceProperties.Add(SourceProperty);
+		}
+	}
+	if (SourceProperties.IsEmpty())
+	{
+		return FECACommandResult::Error(TEXT("Missing required parameter: source_property or source_properties"));
+	}
+	SourceProperty = SourceProperties[0];
+
+	// Optional per-source viewmodel names (for cross-VM multi-source bindings)
+	// Must match source_properties length if provided
+	TArray<FString> SourceViewModels;
+	const TArray<TSharedPtr<FJsonValue>>* SourceVMArray = nullptr;
+	if (GetArrayParam(Params, TEXT("source_viewmodels"), SourceVMArray) && SourceVMArray)
+	{
+		for (const TSharedPtr<FJsonValue>& Val : *SourceVMArray)
+		{
+			FString VMName;
+			if (Val->TryGetString(VMName))
+			{
+				SourceViewModels.Add(VMName);
+			}
+		}
+	}
+
+	// Validate source_viewmodels length matches source_properties if both provided
+	if (!SourceViewModels.IsEmpty() && SourceViewModels.Num() != SourceProperties.Num())
+	{
+		return FECACommandResult::Error(FString::Printf(
+			TEXT("source_viewmodels (%d entries) must have the same length as source_properties (%d entries)"),
+			SourceViewModels.Num(), SourceProperties.Num()));
 	}
 
 	FString DestinationWidget = TEXT("self");
@@ -916,18 +1009,37 @@ FECACommandResult FECACommand_AddViewmodelBinding::Execute(const TSharedPtr<FJso
 		return FECACommandResult::Error(TEXT("ViewModel class is null."));
 	}
 
-	// Resolve source property on ViewModel (FProperty expected, UFunction as fallback)
-	FProperty* VMProp = FindPropertyOnClass(VMClass, SourceProperty);
+	// Resolve the effective ViewModel and source property/function for the first source.
+	// When source_viewmodels[0] is set, the property lives on the override VM — resolve it there
+	// so both the direct-binding path and the conversion-function loop have a valid VMField and VMContext.
+	const bool bFirstSourceHasOverrideVM = SourceViewModels.IsValidIndex(0) && !SourceViewModels[0].IsEmpty();
+	const FMVVMBlueprintViewModelContext* EffectiveVMContext = VMContext;
+	if (bFirstSourceHasOverrideVM)
+	{
+		const FMVVMBlueprintViewModelContext* OverrideCtx = View->FindViewModel(FName(*SourceViewModels[0]));
+		if (!OverrideCtx)
+		{
+			return FECACommandResult::Error(FString::Printf(
+				TEXT("source_viewmodels[0]: ViewModel '%s' not found on widget"), *SourceViewModels[0]));
+		}
+		EffectiveVMContext = OverrideCtx;
+	}
+	const UClass* EffectiveVMClass = EffectiveVMContext->GetViewModelClass();
+	if (!EffectiveVMClass)
+	{
+		return FECACommandResult::Error(TEXT("Effective ViewModel has no resolved class"));
+	}
+	FProperty* VMProp = FindPropertyOnClass(EffectiveVMClass, SourceProperty);
 	UFunction* VMFunc = nullptr;
 	if (!VMProp)
 	{
-		VMFunc = VMClass->FindFunctionByName(*SourceProperty);
+		VMFunc = EffectiveVMClass->FindFunctionByName(*SourceProperty);
 	}
 	if (!VMProp && !VMFunc)
 	{
 		return FECACommandResult::Error(FString::Printf(
 			TEXT("ViewModel property not found: '%s' on %s. Check property name spelling."),
-			*SourceProperty, *VMClass->GetName()));
+			*SourceProperty, *EffectiveVMClass->GetName()));
 	}
 	UE::MVVM::FMVVMConstFieldVariant VMField = VMProp
 		? UE::MVVM::FMVVMConstFieldVariant(VMProp)
@@ -985,7 +1097,7 @@ FECACommandResult FECACommand_AddViewmodelBinding::Execute(const TSharedPtr<FJso
 
 	// Build and set source path (ViewModel.Property)
 	FMVVMBlueprintPropertyPath SourcePath;
-	SourcePath.SetViewModelId(VMContext->GetViewModelId());
+	SourcePath.SetViewModelId(EffectiveVMContext->GetViewModelId());
 	SourcePath.SetPropertyPath(WidgetBP, VMField);
 
 	// Build and set destination path (Widget.Property or Self.Property)
@@ -1027,16 +1139,95 @@ FECACommandResult FECACommand_AddViewmodelBinding::Execute(const TSharedPtr<FJso
 		FMVVMBlueprintFunctionReference FuncRef(WidgetBP, ConvFunc);
 		MVVMSubsystem->SetSourceToDestinationConversionFunction(WidgetBP, NewBinding, FuncRef);
 
-		// Wire the ViewModel property to the conversion function's input pin
+		// Wire all source properties to the conversion function's input pins
 		UE::MVVM::FConversionFunctionValue ConvValue(ConvFunc);
-		TArray<FMVVMBlueprintPinId> InputPins;
-		if (VMProp)
+		for (int32 SourceIdx = 0; SourceIdx < SourceProperties.Num(); ++SourceIdx)
 		{
-			InputPins = MVVMSubsystem->GetConversionFunctionPins(WidgetBP, ConvValue, VMProp);
-		}
-		if (InputPins.Num() > 0)
-		{
-			MVVMSubsystem->SetPathForConversionFunctionArgument(WidgetBP, NewBinding, InputPins[0], SourcePath, true);
+			const FString& SrcProp = SourceProperties[SourceIdx];
+
+			// Check if this source has a per-source VM override.
+			// If so, skip the default VM lookup entirely — property only needs to exist on the override VM.
+			const bool bHasOverrideVM = SourceViewModels.IsValidIndex(SourceIdx) && !SourceViewModels[SourceIdx].IsEmpty();
+
+			// Resolve property/function on the default VM (only when no override is specified)
+			FProperty* SrcVMProp = bHasOverrideVM ? nullptr : FindPropertyOnClass(VMClass, SrcProp);
+			UFunction* SrcVMFunc = nullptr;
+			if (!bHasOverrideVM && !SrcVMProp)
+			{
+				SrcVMFunc = VMClass ? VMClass->FindFunctionByName(*SrcProp) : nullptr;
+			}
+			if (!bHasOverrideVM && !SrcVMProp && !SrcVMFunc)
+			{
+				UE_LOG(LogTemp, Warning, TEXT("[ECABridge] AddViewmodelBinding: source_properties[%d] '%s' not found on ViewModel '%s' — skipping"),
+					SourceIdx, *SrcProp, *VMContext->GetViewModelName().ToString());
+				continue;
+			}
+
+			UE::MVVM::FMVVMConstFieldVariant SrcField = SrcVMProp
+				? UE::MVVM::FMVVMConstFieldVariant(SrcVMProp)
+				: UE::MVVM::FMVVMConstFieldVariant(SrcVMFunc);
+
+			// Resolve ViewModel for this source — use per-source VM if provided
+			const FMVVMBlueprintViewModelContext* SrcVMContext = VMContext;
+			if (SourceViewModels.IsValidIndex(SourceIdx) && !SourceViewModels[SourceIdx].IsEmpty())
+			{
+				const FMVVMBlueprintViewModelContext* OverrideVM = View->FindViewModel(FName(*SourceViewModels[SourceIdx]));
+				if (!OverrideVM)
+				{
+					// Named override VM not found — fail rather than silently bind to the wrong ViewModel
+					MVVMSubsystem->RemoveBinding(WidgetBP, NewBinding);
+					return FECACommandResult::Error(FString::Printf(
+						TEXT("source_viewmodels[%d]: ViewModel '%s' not found on widget"),
+						SourceIdx, *SourceViewModels[SourceIdx]));
+				}
+				SrcVMContext = OverrideVM;
+				// Re-resolve property on the override VM class
+				const UClass* OverrideVMClass = SrcVMContext->GetViewModelClass();
+				if (!OverrideVMClass)
+				{
+					MVVMSubsystem->RemoveBinding(WidgetBP, NewBinding);
+					return FECACommandResult::Error(FString::Printf(
+						TEXT("source_viewmodels[%d]: ViewModel '%s' has no resolved class"),
+						SourceIdx, *SourceViewModels[SourceIdx]));
+				}
+				SrcVMProp = FindPropertyOnClass(OverrideVMClass, SrcProp);
+				SrcVMFunc = nullptr; // Reset — prevents stale func from outer scope
+				if (!SrcVMProp)
+				{
+					SrcVMFunc = OverrideVMClass->FindFunctionByName(*SrcProp);
+				}
+				if (!SrcVMProp && !SrcVMFunc)
+				{
+					MVVMSubsystem->RemoveBinding(WidgetBP, NewBinding);
+					return FECACommandResult::Error(FString::Printf(
+						TEXT("source_viewmodels[%d]: property or function '%s' not found on ViewModel '%s'"),
+						SourceIdx, *SrcProp, *SourceViewModels[SourceIdx]));
+				}
+				SrcField = SrcVMProp
+					? UE::MVVM::FMVVMConstFieldVariant(SrcVMProp)
+					: UE::MVVM::FMVVMConstFieldVariant(SrcVMFunc);
+			}
+
+			FMVVMBlueprintPropertyPath SrcPath;
+			SrcPath.SetViewModelId(SrcVMContext->GetViewModelId());
+			SrcPath.SetPropertyPath(WidgetBP, SrcField);
+
+			// Get ALL conversion function input pins (unfiltered) and index by SourceIdx.
+			// Passing nullptr returns all pins regardless of type, which is correct for
+			// heterogeneous-input functions (e.g. bool + int) where type-filtered results
+			// would be a subset and indexing by SourceIdx would pick the wrong pin.
+			TArray<FMVVMBlueprintPinId> InputPins = MVVMSubsystem->GetConversionFunctionPins(WidgetBP, ConvValue, nullptr);
+			if (InputPins.IsValidIndex(SourceIdx))
+			{
+				MVVMSubsystem->SetPathForConversionFunctionArgument(WidgetBP, NewBinding, InputPins[SourceIdx], SrcPath, true);
+			}
+			else
+			{
+				MVVMSubsystem->RemoveBinding(WidgetBP, NewBinding);
+				return FECACommandResult::Error(FString::Printf(
+					TEXT("source_properties[%d]: conversion function has only %d input pins"),
+					SourceIdx, InputPins.Num()));
+			}
 		}
 
 		ConversionFunctionUsed = ConvFunc->GetName();
@@ -1050,7 +1241,8 @@ FECACommandResult FECACommand_AddViewmodelBinding::Execute(const TSharedPtr<FJso
 	// Set destination path
 	MVVMSubsystem->SetDestinationPathForBinding(WidgetBP, NewBinding, DestPath, false);
 
-	// Capture binding ID before compile (reference may shift during compilation)
+	// Capture binding ID before compile — NewBinding reference is INVALID after CompileAndMarkDirty.
+	// All subsequent access must go through CompiledBinding (re-fetched via NewBindingId below).
 	FGuid NewBindingId = NewBinding.BindingId;
 
 	CompileAndMarkDirty(WidgetBP);
@@ -1114,3 +1306,148 @@ FECACommandResult FECACommand_AddViewmodelBinding::Execute(const TSharedPtr<FJso
 	}
 	return FECACommandResult::Success(Result);
 }
+
+//------------------------------------------------------------------------------
+// SetViewmodelBindingArgument
+//------------------------------------------------------------------------------
+
+FECACommandResult FECACommand_SetViewmodelBindingArgument::Execute(const TSharedPtr<FJsonObject>& Params)
+{
+	FString WidgetPath;
+	if (!GetStringParam(Params, TEXT("widget_path"), WidgetPath))
+	{
+		return FECACommandResult::Error(TEXT("Missing required parameter: widget_path"));
+	}
+
+	double BindingIndexD = 0, ArgIndexD = 0;
+	if (!Params->TryGetNumberField(TEXT("binding_index"), BindingIndexD) ||
+		!Params->TryGetNumberField(TEXT("arg_index"), ArgIndexD))
+	{
+		return FECACommandResult::Error(TEXT("Missing required parameters: binding_index, arg_index"));
+	}
+	int32 BindingIndex = static_cast<int32>(BindingIndexD);
+	int32 ArgIndex     = static_cast<int32>(ArgIndexD);
+
+	FString SourceProperty, VMName;
+	if (!GetStringParam(Params, TEXT("source_property"), SourceProperty) ||
+		!GetStringParam(Params, TEXT("viewmodel_name"), VMName))
+	{
+		return FECACommandResult::Error(TEXT("Missing required parameters: source_property, viewmodel_name"));
+	}
+
+	UWidgetBlueprint* WidgetBP = nullptr;
+	UMVVMEditorSubsystem* MVVMSubsystem = nullptr;
+	UMVVMBlueprintView* View = nullptr;
+	FString ValidationError = ValidateMVVMWidget(WidgetPath, WidgetBP, MVVMSubsystem, View);
+	if (!ValidationError.IsEmpty())
+	{
+		return FECACommandResult::Error(ValidationError);
+	}
+
+	// Find the binding by index
+	if (BindingIndex < 0 || BindingIndex >= View->GetNumBindings())
+	{
+		return FECACommandResult::Error(FString::Printf(
+			TEXT("Binding index %d out of range (0-%d)"), BindingIndex, View->GetNumBindings() - 1));
+	}
+
+	FMVVMBlueprintViewBinding* BindingPtr = View->GetBindingAt(BindingIndex);
+	if (!BindingPtr)
+	{
+		return FECACommandResult::Error(FString::Printf(TEXT("Binding %d not found"), BindingIndex));
+	}
+	// Find the ViewModel
+	const FMVVMBlueprintViewModelContext* VMContext = View->FindViewModel(FName(*VMName));
+	if (!VMContext)
+	{
+		return FECACommandResult::Error(FString::Printf(TEXT("ViewModel '%s' not found"), *VMName));
+	}
+
+	const UClass* VMClass = VMContext->GetViewModelClass();
+	if (!VMClass)
+	{
+		return FECACommandResult::Error(FString::Printf(TEXT("ViewModel '%s' has no resolved class"), *VMName));
+	}
+	FProperty* SrcProp = FindPropertyOnClass(VMClass, SourceProperty);
+	UFunction* SrcFunc = nullptr;
+	if (!SrcProp)
+	{
+		SrcFunc = VMClass->FindFunctionByName(*SourceProperty);
+	}
+	if (!SrcProp && !SrcFunc)
+	{
+		return FECACommandResult::Error(FString::Printf(
+			TEXT("Property or function '%s' not found on ViewModel '%s'"), *SourceProperty, *VMName));
+	}
+
+	UE::MVVM::FMVVMConstFieldVariant SrcField = SrcProp
+		? UE::MVVM::FMVVMConstFieldVariant(SrcProp)
+		: UE::MVVM::FMVVMConstFieldVariant(SrcFunc);
+
+	// Resolve conversion function by name — required param since the binding object's
+	// GetConversionFunction() has an unexported deprecated signature.
+	FString ConvFuncName;
+	GetStringParam(Params, TEXT("conversion_function"), ConvFuncName, false);
+	if (ConvFuncName.IsEmpty())
+	{
+		return FECACommandResult::Error(TEXT("Missing required parameter: conversion_function (the name of the binding's conversion function, e.g. CanTrackFadeOutAtSongEnd)"));
+	}
+
+	FString ConvOwner;
+	UFunction* ResolvedConvFunc = ResolveConversionFunction(WidgetBP, ConvFuncName, ConvOwner);
+
+	UE::MVVM::FConversionFunctionValue ConvValue;
+	if (ResolvedConvFunc)
+	{
+		ConvValue = UE::MVVM::FConversionFunctionValue(ResolvedConvFunc);
+	}
+	else
+	{
+		// Try as K2Node class name
+		for (TObjectIterator<UClass> It; It; ++It)
+		{
+			if (It->IsChildOf(UK2Node::StaticClass()) && It->GetName() == ConvFuncName)
+			{
+				ConvValue = UE::MVVM::FConversionFunctionValue(TSubclassOf<UK2Node>(*It));
+				break;
+			}
+		}
+		if (!ConvValue.IsFunction() && !ConvValue.IsNode())
+		{
+			return FECACommandResult::Error(FString::Printf(TEXT("Conversion function '%s' not found as UFunction or K2Node"), *ConvFuncName));
+		}
+	}
+	// Use unfiltered pin list and index by absolute position (arg_index=0 means first input pin overall).
+	// Type-filtered lookups return a subset and would cause arg_index to select the wrong pin
+	// for heterogeneous-input functions (matching the behavior in AddViewmodelBinding).
+	TArray<FMVVMBlueprintPinId> InputPins = MVVMSubsystem->GetConversionFunctionPins(WidgetBP, ConvValue, nullptr);
+	if (!InputPins.IsValidIndex(ArgIndex))
+	{
+		return FECACommandResult::Error(FString::Printf(
+			TEXT("arg_index %d is out of range (conversion function has %d input pins)"),
+			ArgIndex, InputPins.Num()));
+	}
+
+	// Build source path
+	FMVVMBlueprintPropertyPath SrcPath;
+	SrcPath.SetViewModelId(VMContext->GetViewModelId());
+	SrcPath.SetPropertyPath(WidgetBP, SrcField);
+
+	// Wire it — GetBindingAt(int32) has a non-const overload that returns FMVVMBlueprintViewBinding*
+	FMVVMBlueprintViewBinding* MutableBinding = View->GetBindingAt(BindingIndex);
+	if (!MutableBinding)
+	{
+		return FECACommandResult::Error(TEXT("Could not get mutable binding reference"));
+	}
+	MVVMSubsystem->SetPathForConversionFunctionArgument(WidgetBP, *MutableBinding, InputPins[ArgIndex], SrcPath, true);
+
+	CompileAndMarkDirty(WidgetBP);
+
+	TSharedPtr<FJsonObject> Result = MakeResult();
+	Result->SetNumberField(TEXT("binding_index"), BindingIndex);
+	Result->SetNumberField(TEXT("arg_index"), ArgIndex);
+	Result->SetStringField(TEXT("source"), VMName + TEXT(".") + SourceProperty);
+	Result->SetBoolField(TEXT("success"), true);
+
+	return FECACommandResult::Success(Result);
+}
diff --git a/Engine/Plugins/Experimental/ECABridge/Source/ECABridge/Public/Commands/ECAMVVMCommands.h b/Engine/Plugins/Experimental/ECABridge/Source/ECABridge/Public/Commands/ECAMVVMCommands.h
--- a/Engine/Plugins/Experimental/ECABridge/Source/ECABridge/Public/Commands/ECAMVVMCommands.h
+++ b/Engine/Plugins/Experimental/ECABridge/Source/ECABridge/Public/Commands/ECAMVVMCommands.h
@@ -55,7 +55,7 @@ class FECACommand_BindVisibilityToViewmodel : public IECACommand
 {
 public:
 	virtual FString GetName() const override { return TEXT("bind_visibility_to_viewmodel"); }
-	virtual FString GetDescription() const override { return TEXT("Bind a widget's Visibility to a ViewModel bool via Conv_BoolToSlateVisibility with configurable true/false visibility values"); }
+	virtual FString GetDescription() const override { return TEXT("Bind a widget's Visibility to a ViewModel bool property or FieldNotify getter function via Conv_BoolToSlateVisibility with configurable true/false visibility values"); }
 	virtual FString GetCategory() const override { return TEXT("MVVM"); }
 
 	virtual TArray<FECACommandParam> GetParameters() const override
@@ -63,7 +63,8 @@ public:
 		return {
 			{ TEXT("widget_path"), TEXT("string"), TEXT("Path to the Widget Blueprint asset"), true },
 			{ TEXT("widget_name"), TEXT("string"), TEXT("Name of the widget to bind"), true },
-			{ TEXT("viewmodel_property"), TEXT("string"), TEXT("Name of the ViewModel bool property"), true },
+			{ TEXT("viewmodel_name"), TEXT("string"), TEXT("Name of the ViewModel to use (optional — searches all VMs if omitted)"), false },
+			{ TEXT("viewmodel_property"), TEXT("string"), TEXT("Name of the ViewModel bool property or FieldNotify getter function (e.g. 'bIsActive' or 'GetShouldShowPostUpdateBuffer')"), true },
 			{ TEXT("true_visibility"), TEXT("string"), TEXT("Visibility when bool is true (Visible, Collapsed, Hidden, HitTestInvisible, SelfHitTestInvisible)"), false, TEXT("Visible") },
 			{ TEXT("false_visibility"), TEXT("string"), TEXT("Visibility when bool is false"), false, TEXT("Collapsed") }
 		};
@@ -138,26 +139,28 @@ public:
 };
 
 /**
- * Add a generic MVVM binding between a ViewModel property and any widget property.
- * Supports self-bindings, multi-ViewModel lookup by name, and optional conversion functions.
+ * Add an MVVM binding from one or more ViewModel properties (optionally from different ViewModels) to any widget property.
+ * Supports single-source, multi-source (source_properties[]), cross-VM wiring (source_viewmodels[]), and optional conversion functions.
  */
 class FECACommand_AddViewmodelBinding : public IECACommand
 {
 public:
 	virtual FString GetName() const override { return TEXT("add_viewmodel_binding"); }
-	virtual FString GetDescription() const override { return TEXT("Add a generic MVVM binding from a ViewModel property to any widget property, with optional conversion function"); }
+	virtual FString GetDescription() const override { return TEXT("Add an MVVM binding from one or more ViewModel properties (optionally from different ViewModels) to any widget property, with optional conversion function"); }
 	virtual FString GetCategory() const override { return TEXT("MVVM"); }
 
 	virtual TArray<FECACommandParam> GetParameters() const override
 	{
 		return {
 			{ TEXT("widget_path"), TEXT("string"), TEXT("Path to the Widget Blueprint asset"), true },
-			{ TEXT("source_property"), TEXT("string"), TEXT("ViewModel property name to bind from"), true },
+			{ TEXT("source_property"), TEXT("string"), TEXT("ViewModel property to bind from. Required unless source_properties[] is provided instead."), false },
+			{ TEXT("source_properties"), TEXT("array"), TEXT("Multiple ViewModel property names for multi-input conversion functions (e.g. [\"PropA\", \"PropB\"]). When provided, source_property is not required."), false },
+			{ TEXT("source_viewmodels"), TEXT("array"), TEXT("ViewModel name per source property — use when sources come from different ViewModels (same length as source_properties)"), false },
 			{ TEXT("destination_property"), TEXT("string"), TEXT("Destination property or setter name on the target widget"), true },
 			{ TEXT("destination_widget"), TEXT("string"), TEXT("Child widget name, or 'self' to bind to the widget itself (default: self)"), false, TEXT("self") },
 			{ TEXT("viewmodel_name"), TEXT("string"), TEXT("ViewModel name to bind from (default: first ViewModel)"), false },
 			{ TEXT("binding_mode"), TEXT("string"), TEXT("Binding mode: OneWayToDestination, TwoWay, OneWayToSource, OneTimeToDestination (default: OneWayToDestination)"), false, TEXT("OneWayToDestination") },
-			{ TEXT("conversion_function"), TEXT("string"), TEXT("Optional conversion function (e.g. 'ClassName::FuncName' or 'FuncName' for BP-local)"), false }
+			{ TEXT("conversion_function"), TEXT("string"), TEXT("Conversion function name. Use 'FuncName' for BP-local or 'ClassName::FuncName' to disambiguate across classes (required for K2Node-based functions e.g. CanTrackFadeIn; omit for direct UFUNCTION-based bindings)"), false }
 		};
 	}
 
@@ -184,3 +187,29 @@ public:
 
 	virtual FECACommandResult Execute(const TSharedPtr<FJsonObject>& Params) override;
 };
+
+/**
+ * Wire an additional source property to a specific argument pin of an existing MVVM binding's conversion function.
+ * Use when a binding's conversion function has multiple inputs from different ViewModels.
+ */
+class FECACommand_SetViewmodelBindingArgument : public IECACommand
+{
+public:
+	virtual FString GetName() const override { return TEXT("set_viewmodel_binding_argument"); }
+	virtual FString GetDescription() const override { return TEXT("Wire a source property to a specific argument pin of an existing binding's conversion function. Use for multi-VM multi-source bindings."); }
+	virtual FString GetCategory() const override { return TEXT("MVVM"); }
+
+	virtual TArray<FECACommandParam> GetParameters() const override
+	{
+		return {
+			{ TEXT("widget_path"), TEXT("string"), TEXT("Path to the Widget Blueprint asset"), true },
+			{ TEXT("binding_index"), TEXT("number"), TEXT("Index of the existing binding (from get_viewmodel_bindings)"), true },
+			{ TEXT("arg_index"), TEXT("number"), TEXT("0-based absolute index into the conversion function's input pins (0 = first input pin overall, regardless of type)"), true },
+			{ TEXT("source_property"), TEXT("string"), TEXT("ViewModel property or function name to wire to this argument"), true },
+			{ TEXT("viewmodel_name"), TEXT("string"), TEXT("Name of the ViewModel that owns source_property"), true },
+			{ TEXT("conversion_function"), TEXT("string"), TEXT("Name of the binding's conversion function (required — e.g. CanTrackFadeOutAtSongEnd)"), true },
+		};
+	}
+
+	virtual FECACommandResult Execute(const TSharedPtr<FJsonObject>& Params) override;
+};
