[ECABridge] Widget tree authoring tools: set_image_brush, set_post_buffer_update_config, get_slot_property, create_widget_animation, rename_widget, USpacer size

New MCP commands:
- set_image_brush: full FSlateBrush editor ??? resource, tint, size, draw_as, tiling, mirroring, margin, outline. All properties optional. ParseDrawType/TileType/MirrorType return bool; invalid values produce errors instead of silently defaulting.
- set_post_buffer_update_config: configure PostBufferUpdate widget buffer targets (0-4) and optional Gaussian blur strength. Validates blur_strength > 0. Uses reflection on private UpdateBufferInfos property. NewObject Outer = WidgetBlueprint.
- get_slot_property: read slot layout (padding, alignment, size rule) from any widget. Handles OverlaySlot, VerticalBoxSlot, HorizontalBoxSlot, SizeBoxSlot, StackBoxSlot, CanvasPanelSlot, with reflection fallback.
- create_widget_animation: create empty UWidgetAnimation with Modify() before Animations.Add() for undo support.
- rename_widget: rename with full OnVariableRenamed reference fixup and REN_DontCreateRedirectors.
- set_widget_property: adds USpacer Size support ({X, Y} or single float).

Build.cs: MovieScene (for UWidgetAnimation), DataValidation (for validate_assets) deps added.

#rnx

#Synced-CL 52128568

diff --git a/Engine/Plugins/Experimental/ECABridge/Source/ECABridge/ECABridge.Build.cs b/Engine/Plugins/Experimental/ECABridge/Source/ECABridge/ECABridge.Build.cs
--- a/Engine/Plugins/Experimental/ECABridge/Source/ECABridge/ECABridge.Build.cs
+++ b/Engine/Plugins/Experimental/ECABridge/Source/ECABridge/ECABridge.Build.cs
@@ -56,6 +56,9 @@ public class ECABridge : ModuleRules
 			"MaterialEditor",            // For material editing utilities
 			"DataTableEditor",           // For DataTable editor UI refresh (BroadcastPostChange)
 			
+			"MovieScene",                 // For UWidgetAnimation / UMovieScene in create_widget_animation
+			"DataValidation",             // For UEditorValidatorSubsystem in validate_assets
+
 			// UE AI modules
 			"AIModule",                  // For Behavior Tree and Blackboard asset access
 			"StateTreeModule",           // For UStateTree and StateTree types
@@ -63,7 +66,7 @@ public class ECABridge : ModuleRules
 			
 			// Additional utilities
 			"Projects", "DeveloperSettings", "ToolMenus",
-			
+
 			// Image processing
 			"ImageWrapper",
 			
diff --git a/Engine/Plugins/Experimental/ECABridge/Source/ECABridge/Private/Commands/ECAWidgetTreeCommands.cpp b/Engine/Plugins/Experimental/ECABridge/Source/ECABridge/Private/Commands/ECAWidgetTreeCommands.cpp
--- a/Engine/Plugins/Experimental/ECABridge/Source/ECABridge/Private/Commands/ECAWidgetTreeCommands.cpp
+++ b/Engine/Plugins/Experimental/ECABridge/Source/ECABridge/Private/Commands/ECAWidgetTreeCommands.cpp
@@ -19,20 +19,39 @@
 #include "Components/CanvasPanel.h"
 #include "Components/VerticalBox.h"
 #include "Components/HorizontalBox.h"
+#include "Components/Spacer.h"
 #include "Components/StackBox.h"
+#include "Components/StackBoxSlot.h"
 #include "Components/TextBlock.h"
 #include "Components/Image.h"
 
-// Engine - UMG Slots (SlateWrapperTypes for FSlateChildSize; individual slot headers
-// are NOT needed since we use UE reflection to set slot properties uniformly)
+// Engine - Widget Animation
+#include "Animation/WidgetAnimation.h"
+#include "MovieScene.h"
+
+// Engine - UMG Slots
 #include "Components/PanelSlot.h"
 #include "Components/SlateWrapperTypes.h"
+#include "Components/OverlaySlot.h"
+#include "Components/VerticalBoxSlot.h"
+#include "Components/HorizontalBoxSlot.h"
+#include "Components/SizeBoxSlot.h"
+#include "Components/CanvasPanelSlot.h"
 
 // Engine - Blueprint utilities
 #include "Kismet2/BlueprintEditorUtils.h"
 #include "UObject/UObjectIterator.h"
 #include "UObject/SavePackage.h"
 
+// Engine - Materials (for set_image_brush)
+#include "Materials/MaterialInterface.h"
+#include "Engine/Texture2D.h"
+
+// Engine - PostBufferUpdate (for set_post_buffer_update_config)
+#include "Components/PostBufferUpdate.h"
+#include "Components/PostBufferBlurUpdater.h"
+#include "Rendering/SlateRendererTypes.h"
+
 // Register all widget tree commands
 REGISTER_ECA_COMMAND(FECACommand_WrapWidget)
 REGISTER_ECA_COMMAND(FECACommand_UnwrapWidget)
@@ -42,6 +61,11 @@ REGISTER_ECA_COMMAND(FECACommand_SetWidgetProperty)
 REGISTER_ECA_COMMAND(FECACommand_AddChildWidget)
 REGISTER_ECA_COMMAND(FECACommand_CompileWidgetBlueprint)
 REGISTER_ECA_COMMAND(FECACommand_RemoveWidget)
+REGISTER_ECA_COMMAND(FECACommand_SetImageBrush)
+REGISTER_ECA_COMMAND(FECACommand_SetPostBufferUpdateConfig)
+REGISTER_ECA_COMMAND(FECACommand_RenameWidget)
+REGISTER_ECA_COMMAND(FECACommand_GetSlotProperty)
+REGISTER_ECA_COMMAND(FECACommand_CreateWidgetAnimation)
 
 //------------------------------------------------------------------------------
 // Shared Helper Functions
@@ -1007,6 +1031,40 @@ FECACommandResult FECACommand_SetWidgetProperty::Execute(const TSharedPtr<FJsonO
 			}
 		}
 
+		// --- USpacer properties ---
+		else if (USpacer* Spacer = Cast<USpacer>(Widget))
+		{
+			if (PropName.Equals(TEXT("Size"), ESearchCase::IgnoreCase) ||
+				PropName.Equals(TEXT("DesiredSize"), ESearchCase::IgnoreCase))
+			{
+				const TSharedPtr<FJsonObject>* SizeObj = nullptr;
+				if (PropValue->TryGetObject(SizeObj) && SizeObj && SizeObj->IsValid())
+				{
+					// Initialize from current size — omitted keys preserve the existing value (consistent with scalar path)
+					const FVector2D CurrentSize = Spacer->GetSize();
+					double X = CurrentSize.X, Y = CurrentSize.Y;
+					(*SizeObj)->TryGetNumberField(TEXT("X"), X);
+					(*SizeObj)->TryGetNumberField(TEXT("Y"), Y);
+					// Also accept lowercase
+					(*SizeObj)->TryGetNumberField(TEXT("x"), X);
+					(*SizeObj)->TryGetNumberField(TEXT("y"), Y);
+					Spacer->SetSize(FVector2D(X, Y));
+					bHandled = true;
+				}
+				else
+				{
+					// Accept a single number as the Y value (height)
+					double Val = 0;
+					if (PropValue->TryGetNumber(Val))
+					{
+						FVector2D Current = Spacer->GetSize();
+						Spacer->SetSize(FVector2D(Current.X, Val));
+						bHandled = true;
+					}
+				}
+			}
+		}
+
 		// --- UImage properties ---
 		else if (UImage* Image = Cast<UImage>(Widget))
 		{
@@ -1340,3 +1398,774 @@ FECACommandResult FECACommand_RemoveWidget::Execute(const TSharedPtr<FJsonObject
 	Result->SetNumberField(TEXT("descendants_removed"), DescendantCount);
 	return FECACommandResult::Success(Result);
 }
+
+
+//------------------------------------------------------------------------------
+// SetImageBrush
+//------------------------------------------------------------------------------
+
+static bool ParseDrawType(const FString& Value, TEnumAsByte<ESlateBrushDrawType::Type>& OutType)
+{
+	if (Value.Equals(TEXT("Box"), ESearchCase::IgnoreCase))        { OutType = ESlateBrushDrawType::Box;       return true; }
+	if (Value.Equals(TEXT("Border"), ESearchCase::IgnoreCase))     { OutType = ESlateBrushDrawType::Border;    return true; }
+	if (Value.Equals(TEXT("Image"), ESearchCase::IgnoreCase))      { OutType = ESlateBrushDrawType::Image;     return true; }
+	if (Value.Equals(TEXT("RoundedBox"), ESearchCase::IgnoreCase)) { OutType = ESlateBrushDrawType::RoundedBox;return true; }
+	if (Value.Equals(TEXT("None"), ESearchCase::IgnoreCase))       { OutType = ESlateBrushDrawType::NoDrawType;return true; }
+	return false;
+}
+
+static bool ParseTileType(const FString& Value, TEnumAsByte<ESlateBrushTileType::Type>& OutType)
+{
+	if (Value.Equals(TEXT("Horizontal"), ESearchCase::IgnoreCase)) { OutType = ESlateBrushTileType::Horizontal; return true; }
+	if (Value.Equals(TEXT("Vertical"), ESearchCase::IgnoreCase))   { OutType = ESlateBrushTileType::Vertical;   return true; }
+	if (Value.Equals(TEXT("Both"), ESearchCase::IgnoreCase))       { OutType = ESlateBrushTileType::Both;        return true; }
+	if (Value.Equals(TEXT("NoTile"), ESearchCase::IgnoreCase))     { OutType = ESlateBrushTileType::NoTile;      return true; }
+	return false;
+}
+
+static bool ParseMirrorType(const FString& Value, TEnumAsByte<ESlateBrushMirrorType::Type>& OutType)
+{
+	if (Value.Equals(TEXT("Horizontal"), ESearchCase::IgnoreCase)) { OutType = ESlateBrushMirrorType::Horizontal; return true; }
+	if (Value.Equals(TEXT("Vertical"), ESearchCase::IgnoreCase))   { OutType = ESlateBrushMirrorType::Vertical;   return true; }
+	if (Value.Equals(TEXT("Both"), ESearchCase::IgnoreCase))       { OutType = ESlateBrushMirrorType::Both;        return true; }
+	if (Value.Equals(TEXT("NoMirror"), ESearchCase::IgnoreCase))   { OutType = ESlateBrushMirrorType::NoMirror;   return true; }
+	return false;
+}
+
+
+FECACommandResult FECACommand_SetImageBrush::Execute(const TSharedPtr<FJsonObject>& Params)
+{
+	FString WidgetPath;
+	if (!GetStringParam(Params, TEXT("widget_path"), WidgetPath))
+	{
+		return FECACommandResult::Error(TEXT("Missing required parameter: widget_path"));
+	}
+
+	FString WidgetName;
+	if (!GetStringParam(Params, TEXT("widget_name"), WidgetName))
+	{
+		return FECACommandResult::Error(TEXT("Missing required parameter: widget_name"));
+	}
+
+	// Load the Widget Blueprint
+	UWidgetBlueprint* WidgetBlueprint = LoadWidgetBlueprintFromPath(WidgetPath);
+	if (!WidgetBlueprint)
+	{
+		return FECACommandResult::Error(FString::Printf(TEXT("Widget Blueprint not found: %s"), *WidgetPath));
+	}
+
+	UWidgetTree* WidgetTree = WidgetBlueprint->WidgetTree;
+	if (!WidgetTree)
+	{
+		return FECACommandResult::Error(TEXT("Widget Blueprint has no widget tree"));
+	}
+
+	UWidget* Widget = WidgetTree->FindWidget(FName(*WidgetName));
+	if (!Widget)
+	{
+		return FECACommandResult::Error(FString::Printf(TEXT("Widget not found: %s"), *WidgetName));
+	}
+
+	UImage* Image = Cast<UImage>(Widget);
+	if (!Image)
+	{
+		return FECACommandResult::Error(FString::Printf(TEXT("Widget '%s' is not an Image (class: %s)"), *WidgetName, *Widget->GetClass()->GetName()));
+	}
+
+	TArray<FString> PropertiesSet;
+
+	// Pre-check: at least one property key must be present before we Modify() anything.
+	// This avoids phantom undo entries when the caller provides no optional parameters.
+	const TArray<FString> KnownProps = {
+		TEXT("asset_path"), TEXT("tint"), TEXT("image_size"),
+		TEXT("draw_as"), TEXT("tiling"), TEXT("mirroring"), TEXT("margin"), TEXT("outline")
+	};
+	bool bAnyPropPresent = false;
+	for (const FString& Key : KnownProps)
+	{
+		if (Params->HasField(Key)) { bAnyPropPresent = true; break; }
+	}
+	if (!bAnyPropPresent)
+	{
+		return FECACommandResult::Error(TEXT("No properties specified. Provide at least one of: asset_path, tint, image_size, draw_as, tiling, mirroring, margin, outline"));
+	}
+
+	// Pre-validate asset_path before calling Modify() — avoid phantom undo entries on load failure
+	UObject* Asset = nullptr;
+	FString AssetPath;
+	bool bHasAssetPath = GetStringParam(Params, TEXT("asset_path"), AssetPath, false);
+	if (bHasAssetPath)
+	{
+		Asset = LoadObject<UObject>(nullptr, *AssetPath);
+		if (!Asset)
+		{
+			return FECACommandResult::Error(FString::Printf(TEXT("Could not load asset: %s"), *AssetPath));
+		}
+	}
+
+	// Pre-validate enum strings before Modify() — consistent with asset_path validation above.
+	// Returns early without touching the undo stack if any enum value is unrecognized.
+	FString DrawAsStr, TilingStr, MirroringStr;
+	TEnumAsByte<ESlateBrushDrawType::Type>  DrawAsVal;
+	TEnumAsByte<ESlateBrushTileType::Type>  TilingVal;
+	TEnumAsByte<ESlateBrushMirrorType::Type> MirroringVal;
+	if (GetStringParam(Params, TEXT("draw_as"), DrawAsStr, false) && !ParseDrawType(DrawAsStr, DrawAsVal))
+	{
+		return FECACommandResult::Error(FString::Printf(
+			TEXT("Invalid draw_as: '%s'. Valid: Image, Box, Border, RoundedBox, None"), *DrawAsStr));
+	}
+	if (GetStringParam(Params, TEXT("tiling"), TilingStr, false) && !ParseTileType(TilingStr, TilingVal))
+	{
+		return FECACommandResult::Error(FString::Printf(
+			TEXT("Invalid tiling: '%s'. Valid: NoTile, Horizontal, Vertical, Both"), *TilingStr));
+	}
+	if (GetStringParam(Params, TEXT("mirroring"), MirroringStr, false) && !ParseMirrorType(MirroringStr, MirroringVal))
+	{
+		return FECACommandResult::Error(FString::Printf(
+			TEXT("Invalid mirroring: '%s'. Valid: NoMirror, Horizontal, Vertical, Both"), *MirroringStr));
+	}
+	// rounding_type (inside outline object) — pre-validate if present
+	const TSharedPtr<FJsonObject>* OutlinePreCheck = nullptr;
+	if (Params->TryGetObjectField(TEXT("outline"), OutlinePreCheck) && OutlinePreCheck && OutlinePreCheck->IsValid())
+	{
+		FString RoundingStr;
+		if ((*OutlinePreCheck)->TryGetStringField(TEXT("rounding_type"), RoundingStr))
+		{
+			if (!RoundingStr.Equals(TEXT("HalfHeight"), ESearchCase::IgnoreCase) &&
+				!RoundingStr.Equals(TEXT("FixedRadius"), ESearchCase::IgnoreCase))
+			{
+				return FECACommandResult::Error(FString::Printf(
+					TEXT("Invalid rounding_type: '%s'. Valid: HalfHeight, FixedRadius"), *RoundingStr));
+			}
+		}
+	}
+
+	// Modify() BEFORE any mutations — records pre-change state for undo.
+	// All validatable inputs (asset path + enum strings) are pre-checked above
+	// so no phantom undo entries are created on invalid input.
+	Image->Modify();
+	WidgetBlueprint->Modify();
+
+	// --- Resource (optional — omit to modify brush properties without changing the image) ---
+	if (bHasAssetPath && Asset)
+	{
+
+		bool bMatchSize = false;
+		Params->TryGetBoolField(TEXT("match_size"), bMatchSize);
+
+		if (UMaterialInterface* Material = Cast<UMaterialInterface>(Asset))
+		{
+			Image->SetBrushFromMaterial(Material);
+		}
+		else if (UTexture2D* Texture = Cast<UTexture2D>(Asset))
+		{
+			Image->SetBrushFromTexture(Texture, bMatchSize);
+		}
+		else
+		{
+			Image->SetBrushResourceObject(Asset);
+		}
+		PropertiesSet.Add(TEXT("asset_path"));
+	}
+
+	// --- Tint color ---
+	const TSharedPtr<FJsonObject>* TintObj = nullptr;
+	if (Params->TryGetObjectField(TEXT("tint"), TintObj) && TintObj && TintObj->IsValid())
+	{
+		double R = 1, G = 1, B = 1, A = 1;
+		(*TintObj)->TryGetNumberField(TEXT("r"), R);
+		(*TintObj)->TryGetNumberField(TEXT("g"), G);
+		(*TintObj)->TryGetNumberField(TEXT("b"), B);
+		(*TintObj)->TryGetNumberField(TEXT("a"), A);
+		Image->SetBrushTintColor(FSlateColor(FLinearColor(R, G, B, A)));
+		PropertiesSet.Add(TEXT("tint"));
+	}
+
+	// --- Image size ---
+	const TSharedPtr<FJsonObject>* SizeObj = nullptr;
+	if (Params->TryGetObjectField(TEXT("image_size"), SizeObj) && SizeObj && SizeObj->IsValid())
+	{
+		double W = 0, H = 0;
+		(*SizeObj)->TryGetNumberField(TEXT("x"), W);
+		(*SizeObj)->TryGetNumberField(TEXT("y"), H);
+		Image->SetDesiredSizeOverride(FVector2D(W, H));
+		PropertiesSet.Add(TEXT("image_size"));
+	}
+
+	// --- Brush struct properties (draw type, tiling, mirroring, margin, outline) ---
+	// Copy brush AFTER the asset block so the ResourceObject set by SetBrushFromTexture/SetBrushFromMaterial
+	// is captured in ModifiedBrush before we apply structural changes. Otherwise SetBrush() would
+	// overwrite the resource with the stale pre-asset copy.
+	bool bBrushModified = false;
+	FSlateBrush ModifiedBrush = Image->GetBrush();
+
+	if (!DrawAsStr.IsEmpty())
+	{
+		ModifiedBrush.DrawAs = DrawAsVal;
+		PropertiesSet.Add(TEXT("draw_as"));
+		bBrushModified = true;
+	}
+
+	if (!TilingStr.IsEmpty())
+	{
+		ModifiedBrush.Tiling = TilingVal;
+		PropertiesSet.Add(TEXT("tiling"));
+		bBrushModified = true;
+	}
+
+	if (!MirroringStr.IsEmpty())
+	{
+		ModifiedBrush.Mirroring = MirroringVal;
+		PropertiesSet.Add(TEXT("mirroring"));
+		bBrushModified = true;
+	}
+
+	const TSharedPtr<FJsonObject>* MarginObj = nullptr;
+	if (Params->TryGetObjectField(TEXT("margin"), MarginObj) && MarginObj && MarginObj->IsValid())
+	{
+		double Left = 0, Top = 0, Right = 0, Bottom = 0;
+		(*MarginObj)->TryGetNumberField(TEXT("left"), Left);
+		(*MarginObj)->TryGetNumberField(TEXT("top"), Top);
+		(*MarginObj)->TryGetNumberField(TEXT("right"), Right);
+		(*MarginObj)->TryGetNumberField(TEXT("bottom"), Bottom);
+		ModifiedBrush.Margin = FMargin(Left, Top, Right, Bottom);
+		PropertiesSet.Add(TEXT("margin"));
+		bBrushModified = true;
+	}
+
+	const TSharedPtr<FJsonObject>* OutlineObj = nullptr;
+	if (Params->TryGetObjectField(TEXT("outline"), OutlineObj) && OutlineObj && OutlineObj->IsValid())
+	{
+		FSlateBrushOutlineSettings& Outline = ModifiedBrush.OutlineSettings;
+
+		double Radius = 0;
+		if ((*OutlineObj)->TryGetNumberField(TEXT("corner_radius"), Radius))
+		{
+			Outline.CornerRadii = FVector4(Radius, Radius, Radius, Radius);
+		}
+
+		const TSharedPtr<FJsonObject>* RadiiObj = nullptr;
+		if ((*OutlineObj)->TryGetObjectField(TEXT("corner_radii"), RadiiObj) && RadiiObj && RadiiObj->IsValid())
+		{
+			double TL = 0, TR = 0, BR = 0, BL = 0;
+			(*RadiiObj)->TryGetNumberField(TEXT("top_left"), TL);
+			(*RadiiObj)->TryGetNumberField(TEXT("top_right"), TR);
+			(*RadiiObj)->TryGetNumberField(TEXT("bottom_right"), BR);
+			(*RadiiObj)->TryGetNumberField(TEXT("bottom_left"), BL);
+			Outline.CornerRadii = FVector4(TL, TR, BR, BL);
+		}
+
+		const TSharedPtr<FJsonObject>* OutColorObj = nullptr;
+		if ((*OutlineObj)->TryGetObjectField(TEXT("color"), OutColorObj) && OutColorObj && OutColorObj->IsValid())
+		{
+			double R = 1, G = 1, B = 1, A = 1;
+			(*OutColorObj)->TryGetNumberField(TEXT("r"), R);
+			(*OutColorObj)->TryGetNumberField(TEXT("g"), G);
+			(*OutColorObj)->TryGetNumberField(TEXT("b"), B);
+			(*OutColorObj)->TryGetNumberField(TEXT("a"), A);
+			Outline.Color = FSlateColor(FLinearColor(R, G, B, A));
+		}
+
+		double Width = 0;
+		if ((*OutlineObj)->TryGetNumberField(TEXT("width"), Width))
+		{
+			Outline.Width = static_cast<float>(Width);
+		}
+
+		FString RoundingStr;
+		if (GetStringParam(*OutlineObj, TEXT("rounding_type"), RoundingStr, false))
+		{
+			if (RoundingStr.Equals(TEXT("HalfHeight"), ESearchCase::IgnoreCase))
+			{
+				Outline.RoundingType = ESlateBrushRoundingType::HalfHeightRadius;
+			}
+			else if (RoundingStr.Equals(TEXT("FixedRadius"), ESearchCase::IgnoreCase))
+			{
+				Outline.RoundingType = ESlateBrushRoundingType::FixedRadius;
+			}
+			else
+			{
+				return FECACommandResult::Error(FString::Printf(
+					TEXT("Invalid rounding_type: '%s'. Valid: HalfHeight, FixedRadius"), *RoundingStr));
+			}
+		}
+
+		PropertiesSet.Add(TEXT("outline"));
+		bBrushModified = true;
+	}
+
+	if (bBrushModified)
+	{
+		Image->SetBrush(ModifiedBrush);
+	}
+
+	FBlueprintEditorUtils::MarkBlueprintAsModified(WidgetBlueprint);
+
+	TSharedPtr<FJsonObject> Result = MakeResult();
+	Result->SetStringField(TEXT("widget"), WidgetName);
+	TArray<TSharedPtr<FJsonValue>> SetArray;
+	for (const FString& Prop : PropertiesSet)
+	{
+		SetArray.Add(MakeShared<FJsonValueString>(Prop));
+	}
+	Result->SetArrayField(TEXT("properties_set"), SetArray);
+	return FECACommandResult::Success(Result);
+}
+
+//------------------------------------------------------------------------------
+// SetPostBufferUpdateConfig
+//------------------------------------------------------------------------------
+
+FECACommandResult FECACommand_SetPostBufferUpdateConfig::Execute(const TSharedPtr<FJsonObject>& Params)
+{
+	FString WidgetPath;
+	if (!GetStringParam(Params, TEXT("widget_path"), WidgetPath))
+	{
+		return FECACommandResult::Error(TEXT("Missing required parameter: widget_path"));
+	}
+
+	FString WidgetName;
+	if (!GetStringParam(Params, TEXT("widget_name"), WidgetName))
+	{
+		return FECACommandResult::Error(TEXT("Missing required parameter: widget_name"));
+	}
+
+	const TArray<TSharedPtr<FJsonValue>>* BuffersArray = nullptr;
+	if (!GetArrayParam(Params, TEXT("buffers"), BuffersArray))
+	{
+		return FECACommandResult::Error(TEXT("Missing required parameter: buffers (array of buffer indices 0-4)"));
+	}
+
+	// Load the Widget Blueprint
+	UWidgetBlueprint* WidgetBlueprint = LoadWidgetBlueprintFromPath(WidgetPath);
+	if (!WidgetBlueprint)
+	{
+		return FECACommandResult::Error(FString::Printf(TEXT("Widget Blueprint not found: %s"), *WidgetPath));
+	}
+
+	UWidgetTree* WidgetTree = WidgetBlueprint->WidgetTree;
+	if (!WidgetTree)
+	{
+		return FECACommandResult::Error(TEXT("Widget Blueprint has no widget tree"));
+	}
+
+	// Find the widget
+	UWidget* Widget = WidgetTree->FindWidget(FName(*WidgetName));
+	if (!Widget)
+	{
+		return FECACommandResult::Error(FString::Printf(TEXT("Widget not found: %s"), *WidgetName));
+	}
+
+	UPostBufferUpdate* PostBuffer = Cast<UPostBufferUpdate>(Widget);
+	if (!PostBuffer)
+	{
+		return FECACommandResult::Error(FString::Printf(TEXT("Widget '%s' is not a PostBufferUpdate (class: %s)"), *WidgetName, *Widget->GetClass()->GetName()));
+	}
+
+	// Parse optional blur strength
+	double BlurStrength = -1.0;
+	bool bHasBlurStrength = false;
+	if (Params->TryGetNumberField(TEXT("blur_strength"), BlurStrength))
+	{
+		if (BlurStrength <= 0.0)
+		{
+			return FECACommandResult::Error(FString::Printf(TEXT("blur_strength must be positive (got %f)"), BlurStrength));
+		}
+		bHasBlurStrength = true;
+	}
+
+	// Map of valid buffer indices to enum values
+	static const ESlatePostRT BufferEnumValues[] = {
+		ESlatePostRT::ESlatePostRT_0,
+		ESlatePostRT::ESlatePostRT_1,
+		ESlatePostRT::ESlatePostRT_2,
+		ESlatePostRT::ESlatePostRT_3,
+		ESlatePostRT::ESlatePostRT_4,
+	};
+	static_assert(UE_ARRAY_COUNT(BufferEnumValues) == 5, "BufferEnumValues size mismatch");
+
+	// Modify() BEFORE NewObject loop — UPostBufferBlurUpdater instances are created inside the loop.
+	// Calling Modify() first ensures the undo record encompasses their creation,
+	// preventing orphaned UObjects on undo (same principle as CreateWidgetAnimation).
+	PostBuffer->Modify();
+	WidgetBlueprint->Modify();
+
+	// Build UpdateBufferInfos
+	TArray<FSlatePostBufferUpdateInfo> NewBufferInfos;
+	TArray<TSharedPtr<FJsonValue>> BufferIndicesForResult;
+
+	for (const TSharedPtr<FJsonValue>& BufferVal : *BuffersArray)
+	{
+		if (BufferVal->Type != EJson::Number)
+		{
+			return FECACommandResult::Error(TEXT("Each element in buffers array must be a number (0-4)"));
+		}
+
+		int32 BufferIndex = static_cast<int32>(BufferVal->AsNumber());
+		if (BufferIndex < 0 || BufferIndex > 4)
+		{
+			return FECACommandResult::Error(FString::Printf(TEXT("Invalid buffer index: %d (must be 0-4)"), BufferIndex));
+		}
+
+		FSlatePostBufferUpdateInfo Info;
+		Info.BufferToUpdate = BufferEnumValues[BufferIndex];
+
+		// If blur strength specified, create a PostBufferBlurUpdater for each buffer
+		if (bHasBlurStrength)
+		{
+			UPostBufferBlurUpdater* BlurUpdater = NewObject<UPostBufferBlurUpdater>(WidgetBlueprint);
+			BlurUpdater->GaussianBlurStrength = static_cast<float>(BlurStrength);
+			Info.PostParamUpdater = BlurUpdater;
+		}
+
+		NewBufferInfos.Add(Info);
+		BufferIndicesForResult.Add(MakeShared<FJsonValueNumber>(BufferIndex));
+	}
+
+	// Apply via UE reflection — UpdateBufferInfos is private with AllowPrivateAccess
+	FProperty* Prop = UPostBufferUpdate::StaticClass()->FindPropertyByName(TEXT("UpdateBufferInfos"));
+	if (!Prop)
+	{
+		return FECACommandResult::Error(TEXT("Failed to find UpdateBufferInfos property via reflection"));
+	}
+
+	// Validate property type before casting — guards against engine version changes where
+	// UpdateBufferInfos could have a different type despite matching the property name.
+	FArrayProperty* ArrayProp = CastField<FArrayProperty>(Prop);
+	if (!ArrayProp)
+	{
+		return FECACommandResult::Error(TEXT("UpdateBufferInfos property is not a TArray (engine version mismatch?)"));
+	}
+	FStructProperty* InnerProp = CastField<FStructProperty>(ArrayProp->Inner);
+	if (!InnerProp || InnerProp->Struct != FSlatePostBufferUpdateInfo::StaticStruct())
+	{
+		return FECACommandResult::Error(TEXT("UpdateBufferInfos inner type is not FSlatePostBufferUpdateInfo (engine version mismatch?)"));
+	}
+	TArray<FSlatePostBufferUpdateInfo>* TargetArray = ArrayProp->ContainerPtrToValuePtr<TArray<FSlatePostBufferUpdateInfo>>(PostBuffer);
+	const int32 BufferCount = NewBufferInfos.Num();
+	*TargetArray = MoveTemp(NewBufferInfos);
+
+	// Notify the object of the property change so any editor-time caching in UPostBufferUpdate is refreshed
+	FPropertyChangedEvent PropChangedEvent(Prop, EPropertyChangeType::ValueSet);
+	PostBuffer->PostEditChangeProperty(PropChangedEvent);
+
+	FBlueprintEditorUtils::MarkBlueprintAsModified(WidgetBlueprint);
+
+	// Build result
+	TSharedPtr<FJsonObject> Result = MakeResult();
+	Result->SetStringField(TEXT("widget"), WidgetName);
+	Result->SetArrayField(TEXT("buffers_configured"), BufferIndicesForResult);
+	if (BufferCount == 0)
+	{
+		Result->SetStringField(TEXT("note"), TEXT("Cleared all buffer configurations"));
+	}
+	if (bHasBlurStrength)
+	{
+		Result->SetNumberField(TEXT("blur_strength"), BlurStrength);
+	}
+	return FECACommandResult::Success(Result);
+}
+
+//------------------------------------------------------------------------------
+// RenameWidget
+//------------------------------------------------------------------------------
+
+FECACommandResult FECACommand_RenameWidget::Execute(const TSharedPtr<FJsonObject>& Params)
+{
+	FString WidgetPath;
+	if (!GetStringParam(Params, TEXT("widget_path"), WidgetPath))
+	{
+		return FECACommandResult::Error(TEXT("Missing required parameter: widget_path"));
+	}
+
+	FString WidgetName;
+	if (!GetStringParam(Params, TEXT("widget_name"), WidgetName))
+	{
+		return FECACommandResult::Error(TEXT("Missing required parameter: widget_name"));
+	}
+
+	FString NewName;
+	if (!GetStringParam(Params, TEXT("new_name"), NewName))
+	{
+		return FECACommandResult::Error(TEXT("Missing required parameter: new_name"));
+	}
+
+	UWidgetBlueprint* WidgetBlueprint = LoadWidgetBlueprintFromPath(WidgetPath);
+	if (!WidgetBlueprint)
+	{
+		return FECACommandResult::Error(FString::Printf(TEXT("Widget Blueprint not found: %s"), *WidgetPath));
+	}
+
+	UWidgetTree* WidgetTree = WidgetBlueprint->WidgetTree;
+	if (!WidgetTree)
+	{
+		return FECACommandResult::Error(TEXT("Widget Blueprint has no widget tree"));
+	}
+
+	UWidget* Widget = WidgetTree->FindWidget(FName(*WidgetName));
+	if (!Widget)
+	{
+		return FECACommandResult::Error(FString::Printf(TEXT("Widget not found: %s"), *WidgetName));
+	}
+
+	const FName OldFName = Widget->GetFName();
+	const FName NewFName(*NewName);
+
+	// Reject if new name already taken
+	if (WidgetTree->FindWidget(NewFName))
+	{
+		return FECACommandResult::Error(FString::Printf(TEXT("A widget named '%s' already exists in this Blueprint"), *NewName));
+	}
+
+	WidgetBlueprint->Modify();
+	Widget->Modify();
+
+	// Rename the widget object first — if this fails, abort before touching variable references
+	if (!Widget->Rename(*NewName, nullptr, REN_DontCreateRedirectors))
+	{
+		return FECACommandResult::Error(FString::Printf(
+			TEXT("Failed to rename widget object from '%s' to '%s' (invalid name or collision at UObject level)"),
+			*WidgetName, *NewName));
+	}
+
+	// Update variable references in Blueprint graphs now that the rename succeeded
+	WidgetBlueprint->OnVariableRenamed(OldFName, NewFName);
+
+	// Keep WidgetVariableNameToGuidMap consistent: remove the old-name entry.
+	// OnVariableRenamed updates graph references but does not clean up the stale map key.
+	// EnsureAllWidgetGuids() then adds NewFName -> GUID so the compiler's ValidateAndFixUpVariableGuids
+	// finds a valid entry for the renamed widget.
+	WidgetBlueprint->WidgetVariableNameToGuidMap.Remove(OldFName);
+	EnsureAllWidgetGuids(WidgetBlueprint);
+
+	FBlueprintEditorUtils::MarkBlueprintAsStructurallyModified(WidgetBlueprint);
+
+	TSharedPtr<FJsonObject> Result = MakeResult();
+	Result->SetStringField(TEXT("widget_path"), WidgetPath);
+	Result->SetStringField(TEXT("old_name"), WidgetName);
+	Result->SetStringField(TEXT("new_name"), NewName);
+	Result->SetBoolField(TEXT("success"), true);
+
+	return FECACommandResult::Success(Result);
+}
+
+//------------------------------------------------------------------------------
+// GetSlotProperty
+//------------------------------------------------------------------------------
+
+FECACommandResult FECACommand_GetSlotProperty::Execute(const TSharedPtr<FJsonObject>& Params)
+{
+	FString WidgetPath;
+	if (!GetStringParam(Params, TEXT("widget_path"), WidgetPath))
+	{
+		return FECACommandResult::Error(TEXT("Missing required parameter: widget_path"));
+	}
+
+	FString WidgetName;
+	if (!GetStringParam(Params, TEXT("widget_name"), WidgetName))
+	{
+		return FECACommandResult::Error(TEXT("Missing required parameter: widget_name"));
+	}
+
+	UWidgetBlueprint* WidgetBlueprint = LoadWidgetBlueprintFromPath(WidgetPath);
+	if (!WidgetBlueprint)
+	{
+		return FECACommandResult::Error(FString::Printf(TEXT("Widget Blueprint not found: %s"), *WidgetPath));
+	}
+
+	UWidgetTree* WidgetTree = WidgetBlueprint->WidgetTree;
+	if (!WidgetTree)
+	{
+		return FECACommandResult::Error(TEXT("Widget Blueprint has no widget tree"));
+	}
+
+	UWidget* Widget = WidgetTree->FindWidget(FName(*WidgetName));
+	if (!Widget)
+	{
+		return FECACommandResult::Error(FString::Printf(TEXT("Widget not found: %s"), *WidgetName));
+	}
+
+	UPanelSlot* Slot = Widget->Slot;
+	if (!Slot)
+	{
+		return FECACommandResult::Error(FString::Printf(TEXT("Widget '%s' has no slot (it may be the root widget)"), *WidgetName));
+	}
+
+	TSharedPtr<FJsonObject> Result = MakeResult();
+	Result->SetStringField(TEXT("widget"), WidgetName);
+	Result->SetStringField(TEXT("slot_type"), Slot->GetClass()->GetName());
+
+	auto MarginToJson = [](const FMargin& M) -> TSharedPtr<FJsonObject>
+	{
+		TSharedPtr<FJsonObject> J = MakeShared<FJsonObject>();
+		J->SetNumberField(TEXT("left"),   M.Left);
+		J->SetNumberField(TEXT("top"),    M.Top);
+		J->SetNumberField(TEXT("right"),  M.Right);
+		J->SetNumberField(TEXT("bottom"), M.Bottom);
+		return J;
+	};
+
+	auto HAlignStr = [](EHorizontalAlignment A) -> FString
+	{
+		switch (A)
+		{
+			case HAlign_Fill:   return TEXT("Fill");
+			case HAlign_Left:   return TEXT("Left");
+			case HAlign_Center: return TEXT("Center");
+			case HAlign_Right:  return TEXT("Right");
+			default:            return TEXT("Fill");
+		}
+	};
+
+	auto VAlignStr = [](EVerticalAlignment A) -> FString
+	{
+		switch (A)
+		{
+			case VAlign_Fill:   return TEXT("Fill");
+			case VAlign_Top:    return TEXT("Top");
+			case VAlign_Center: return TEXT("Center");
+			case VAlign_Bottom: return TEXT("Bottom");
+			default:            return TEXT("Fill");
+		}
+	};
+
+	if (UOverlaySlot* OverlayS = Cast<UOverlaySlot>(Slot))
+	{
+		Result->SetObjectField(TEXT("padding"),              MarginToJson(OverlayS->GetPadding()));
+		Result->SetStringField(TEXT("horizontal_alignment"), HAlignStr(OverlayS->GetHorizontalAlignment()));
+		Result->SetStringField(TEXT("vertical_alignment"),   VAlignStr(OverlayS->GetVerticalAlignment()));
+	}
+	else if (UVerticalBoxSlot* VBS = Cast<UVerticalBoxSlot>(Slot))
+	{
+		Result->SetObjectField(TEXT("padding"),              MarginToJson(VBS->GetPadding()));
+		Result->SetStringField(TEXT("horizontal_alignment"), HAlignStr(VBS->GetHorizontalAlignment()));
+		Result->SetStringField(TEXT("vertical_alignment"),   VAlignStr(VBS->GetVerticalAlignment()));
+		Result->SetStringField(TEXT("size_rule"),   VBS->GetSize().SizeRule == ESlateSizeRule::Fill ? TEXT("Fill") : TEXT("Auto"));
+		Result->SetNumberField(TEXT("size_value"),  VBS->GetSize().Value);
+	}
+	else if (UHorizontalBoxSlot* HBS = Cast<UHorizontalBoxSlot>(Slot))
+	{
+		Result->SetObjectField(TEXT("padding"),              MarginToJson(HBS->GetPadding()));
+		Result->SetStringField(TEXT("horizontal_alignment"), HAlignStr(HBS->GetHorizontalAlignment()));
+		Result->SetStringField(TEXT("vertical_alignment"),   VAlignStr(HBS->GetVerticalAlignment()));
+		Result->SetStringField(TEXT("size_rule"),   HBS->GetSize().SizeRule == ESlateSizeRule::Fill ? TEXT("Fill") : TEXT("Auto"));
+		Result->SetNumberField(TEXT("size_value"),  HBS->GetSize().Value);
+	}
+	else if (UStackBoxSlot* StackS = Cast<UStackBoxSlot>(Slot))
+	{
+		Result->SetObjectField(TEXT("padding"),              MarginToJson(StackS->GetPadding()));
+		Result->SetStringField(TEXT("horizontal_alignment"), HAlignStr(StackS->GetHorizontalAlignment()));
+		Result->SetStringField(TEXT("vertical_alignment"),   VAlignStr(StackS->GetVerticalAlignment()));
+		Result->SetStringField(TEXT("size_rule"),             StackS->GetSize().SizeRule == ESlateSizeRule::Fill ? TEXT("Fill") : TEXT("Auto"));
+		Result->SetNumberField(TEXT("size_value"),            StackS->GetSize().Value);
+	}
+	else if (USizeBoxSlot* SBS = Cast<USizeBoxSlot>(Slot))
+	{
+		Result->SetObjectField(TEXT("padding"),              MarginToJson(SBS->GetPadding()));
+		Result->SetStringField(TEXT("horizontal_alignment"), HAlignStr(SBS->GetHorizontalAlignment()));
+		Result->SetStringField(TEXT("vertical_alignment"),   VAlignStr(SBS->GetVerticalAlignment()));
+	}
+	else if (UCanvasPanelSlot* CanvasS = Cast<UCanvasPanelSlot>(Slot))
+	{
+		const FAnchorData& Layout = CanvasS->GetLayout();
+		TSharedPtr<FJsonObject> OffJ = MakeShared<FJsonObject>();
+		OffJ->SetNumberField(TEXT("left"),   Layout.Offsets.Left);
+		OffJ->SetNumberField(TEXT("top"),    Layout.Offsets.Top);
+		OffJ->SetNumberField(TEXT("right"),  Layout.Offsets.Right);
+		OffJ->SetNumberField(TEXT("bottom"), Layout.Offsets.Bottom);
+		Result->SetObjectField(TEXT("offsets"), OffJ);
+		TSharedPtr<FJsonObject> AnchorJ = MakeShared<FJsonObject>();
+		AnchorJ->SetNumberField(TEXT("min_x"), Layout.Anchors.Minimum.X);
+		AnchorJ->SetNumberField(TEXT("min_y"), Layout.Anchors.Minimum.Y);
+		AnchorJ->SetNumberField(TEXT("max_x"), Layout.Anchors.Maximum.X);
+		AnchorJ->SetNumberField(TEXT("max_y"), Layout.Anchors.Maximum.Y);
+		Result->SetObjectField(TEXT("anchors"), AnchorJ);
+		TSharedPtr<FJsonObject> AlignJ = MakeShared<FJsonObject>();
+		AlignJ->SetNumberField(TEXT("x"), Layout.Alignment.X);
+		AlignJ->SetNumberField(TEXT("y"), Layout.Alignment.Y);
+		Result->SetObjectField(TEXT("alignment"), AlignJ);
+		Result->SetBoolField(TEXT("auto_size"), CanvasS->GetAutoSize());
+		Result->SetNumberField(TEXT("z_order"),  CanvasS->GetZOrder());
+	}
+	else
+	{
+		// Generic: use reflection to read Padding, HAlign, VAlign if they exist (handles UBorderSlot and other unrecognized types)
+		if (FStructProperty* PaddingProp = CastField<FStructProperty>(Slot->GetClass()->FindPropertyByName(TEXT("Padding"))))
+		{
+			if (FMargin* PaddingPtr = PaddingProp->ContainerPtrToValuePtr<FMargin>(Slot))
+			{
+				Result->SetObjectField(TEXT("padding"), MarginToJson(*PaddingPtr));
+			}
+		}
+		if (FByteProperty* HProp = CastField<FByteProperty>(Slot->GetClass()->FindPropertyByName(TEXT("HorizontalAlignment"))))
+		{
+			Result->SetStringField(TEXT("horizontal_alignment"), HAlignStr(static_cast<EHorizontalAlignment>(HProp->GetPropertyValue_InContainer(Slot))));
+		}
+		if (FByteProperty* VProp = CastField<FByteProperty>(Slot->GetClass()->FindPropertyByName(TEXT("VerticalAlignment"))))
+		{
+			Result->SetStringField(TEXT("vertical_alignment"), VAlignStr(static_cast<EVerticalAlignment>(VProp->GetPropertyValue_InContainer(Slot))));
+		}
+		Result->SetStringField(TEXT("note"), FString::Printf(TEXT("Unrecognized slot type: %s"), *Slot->GetClass()->GetName()));
+	}
+
+	return FECACommandResult::Success(Result);
+}
+
+//------------------------------------------------------------------------------
+// CreateWidgetAnimation
+//------------------------------------------------------------------------------
+
+FECACommandResult FECACommand_CreateWidgetAnimation::Execute(const TSharedPtr<FJsonObject>& Params)
+{
+	FString WidgetPath;
+	if (!GetStringParam(Params, TEXT("widget_path"), WidgetPath))
+	{
+		return FECACommandResult::Error(TEXT("Missing required parameter: widget_path"));
+	}
+
+	FString AnimationName;
+	if (!GetStringParam(Params, TEXT("animation_name"), AnimationName))
+	{
+		return FECACommandResult::Error(TEXT("Missing required parameter: animation_name"));
+	}
+
+	UWidgetBlueprint* WidgetBlueprint = LoadWidgetBlueprintFromPath(WidgetPath);
+	if (!WidgetBlueprint)
+	{
+		return FECACommandResult::Error(FString::Printf(TEXT("Widget Blueprint not found: %s"), *WidgetPath));
+	}
+
+	// Check for duplicate name
+	const FName AnimFName(*AnimationName);
+	for (UWidgetAnimation* Existing : WidgetBlueprint->Animations)
+	{
+		if (Existing && Existing->GetFName() == AnimFName)
+		{
+			return FECACommandResult::Error(FString::Printf(TEXT("Animation '%s' already exists on this Widget Blueprint"), *AnimationName));
+		}
+	}
+
+	// Modify() BEFORE NewObject so the undo record encompasses the object creation.
+	// Without this, undo would remove NewAnim from Animations[] but leave the UObject
+	// as an orphan in memory.
+	WidgetBlueprint->Modify();
+
+	// Create the animation object
+	UWidgetAnimation* NewAnim = NewObject<UWidgetAnimation>(WidgetBlueprint, AnimFName, RF_Transactional | RF_Public);
+	if (!NewAnim)
+	{
+		return FECACommandResult::Error(FString::Printf(TEXT("Failed to create animation '%s'"), *AnimationName));
+	}
+
+	NewAnim->MovieScene = NewObject<UMovieScene>(NewAnim, AnimFName, RF_Transactional | RF_Public);
+
+	WidgetBlueprint->Animations.Add(NewAnim);
+	FBlueprintEditorUtils::MarkBlueprintAsStructurallyModified(WidgetBlueprint);
+
+	TSharedPtr<FJsonObject> Result = MakeResult();
+	Result->SetStringField(TEXT("widget_path"), WidgetPath);
+	Result->SetStringField(TEXT("animation_name"), AnimationName);
+	Result->SetBoolField(TEXT("success"), true);
+
+	return FECACommandResult::Success(Result);
+}
diff --git a/Engine/Plugins/Experimental/ECABridge/Source/ECABridge/Public/Commands/ECAWidgetTreeCommands.h b/Engine/Plugins/Experimental/ECABridge/Source/ECABridge/Public/Commands/ECAWidgetTreeCommands.h
--- a/Engine/Plugins/Experimental/ECABridge/Source/ECABridge/Public/Commands/ECAWidgetTreeCommands.h
+++ b/Engine/Plugins/Experimental/ECABridge/Source/ECABridge/Public/Commands/ECAWidgetTreeCommands.h
@@ -118,7 +118,7 @@ public:
 		return {
 			{ TEXT("widget_path"), TEXT("string"), TEXT("Path to the Widget Blueprint"), true },
 			{ TEXT("widget_name"), TEXT("string"), TEXT("Name of the widget to modify"), true },
-			{ TEXT("properties"), TEXT("object"), TEXT("Key-value pairs of properties to set (e.g., {\"WidthOverride\": 400, \"Visibility\": \"Collapsed\"})"), true },
+			{ TEXT("properties"), TEXT("object"), TEXT("Key-value pairs of properties to set (e.g., {\"WidthOverride\": 400, \"Visibility\": \"Collapsed\"}). For USpacer Size: use {\"X\":w,\"Y\":h} for both axes, or a single number to set only the height (Y) while preserving the current width."), true },
 		};
 	}
 
@@ -193,3 +193,123 @@ public:
 
 	virtual FECACommandResult Execute(const TSharedPtr<FJsonObject>& Params) override;
 };
+
+/**
+ * Set an Image widget's brush resource and/or brush properties.
+ * At least one of asset_path, tint, image_size, draw_as, tiling, mirroring, margin, or outline must be provided.
+ */
+class FECACommand_SetImageBrush : public IECACommand
+{
+public:
+	virtual FString GetName() const override { return TEXT("set_image_brush"); }
+	virtual FString GetDescription() const override { return TEXT("Set an Image widget's brush resource, tint, size, draw type, tiling, mirroring, margin, and outline. All brush properties are optional."); }
+	virtual FString GetCategory() const override { return TEXT("UMG"); }
+
+	virtual TArray<FECACommandParam> GetParameters() const override
+	{
+		return {
+			{ TEXT("widget_path"), TEXT("string"), TEXT("Path to the Widget Blueprint"), true },
+			{ TEXT("widget_name"), TEXT("string"), TEXT("Name of the Image widget"), true },
+			{ TEXT("asset_path"), TEXT("string"), TEXT("Content path to material or texture (e.g., /Game/UI/M_MyMaterial)"), false },
+			{ TEXT("match_size"), TEXT("boolean"), TEXT("If true and asset_path is a Texture2D, resize brush to texture dimensions"), false },
+			{ TEXT("tint"), TEXT("object"), TEXT("Brush tint color as {r, g, b, a} (0-1 floats)"), false },
+			{ TEXT("image_size"), TEXT("object"), TEXT("Desired size override as {x, y}"), false },
+			{ TEXT("draw_as"), TEXT("string"), TEXT("Draw type: Image, Box, Border, RoundedBox, None"), false },
+			{ TEXT("tiling"), TEXT("string"), TEXT("Tile mode: NoTile, Horizontal, Vertical, Both"), false },
+			{ TEXT("mirroring"), TEXT("string"), TEXT("Mirror mode: NoMirror, Horizontal, Vertical, Both"), false },
+			{ TEXT("margin"), TEXT("object"), TEXT("9-slice margin as {left, top, right, bottom} (0-1 UV space)"), false },
+			{ TEXT("outline"), TEXT("object"), TEXT("Outline settings for RoundedBox: {corner_radius, corner_radii:{top_left,top_right,bottom_right,bottom_left}, color:{r,g,b,a}, width, rounding_type}"), false },
+		};
+	}
+
+	virtual FECACommandResult Execute(const TSharedPtr<FJsonObject>& Params) override;
+};
+
+/**
+ * Configure a PostBufferUpdate widget's UpdateBufferInfos array.
+ * Sets which SlatePost render targets to capture when this widget is drawn.
+ */
+class FECACommand_SetPostBufferUpdateConfig : public IECACommand
+{
+public:
+	virtual FString GetName() const override { return TEXT("set_post_buffer_update_config"); }
+	virtual FString GetDescription() const override { return TEXT("Configure which SlatePost buffers a PostBufferUpdate widget captures. Set buffer indices (0-4) and optional blur updater strength."); }
+	virtual FString GetCategory() const override { return TEXT("UMG"); }
+
+	virtual TArray<FECACommandParam> GetParameters() const override
+	{
+		return {
+			{ TEXT("widget_path"), TEXT("string"), TEXT("Path to the Widget Blueprint"), true },
+			{ TEXT("widget_name"), TEXT("string"), TEXT("Name of the PostBufferUpdate widget"), true },
+			{ TEXT("buffers"), TEXT("array"), TEXT("Array of buffer indices to capture (0-4), e.g. [0] for SlatePostRT_0"), true },
+			{ TEXT("blur_strength"), TEXT("number"), TEXT("If set, adds a PostBufferBlurUpdater per buffer entry with this gaussian blur strength (must be > 0)"), false },
+		};
+	}
+
+	virtual FECACommandResult Execute(const TSharedPtr<FJsonObject>& Params) override;
+};
+
+/**
+ * Rename a widget in a Widget Blueprint and update all variable references.
+ */
+class FECACommand_RenameWidget : public IECACommand
+{
+public:
+	virtual FString GetName() const override { return TEXT("rename_widget"); }
+	virtual FString GetDescription() const override { return TEXT("Rename a widget in a Widget Blueprint, updating all variable references"); }
+	virtual FString GetCategory() const override { return TEXT("UMG"); }
+
+	virtual TArray<FECACommandParam> GetParameters() const override
+	{
+		return {
+			{ TEXT("widget_path"), TEXT("string"), TEXT("Path to the Widget Blueprint"), true },
+			{ TEXT("widget_name"), TEXT("string"), TEXT("Current name of the widget to rename"), true },
+			{ TEXT("new_name"), TEXT("string"), TEXT("New name for the widget"), true },
+		};
+	}
+
+	virtual FECACommandResult Execute(const TSharedPtr<FJsonObject>& Params) override;
+};
+
+/**
+ * Read slot layout properties for a widget (padding, alignment, size rule).
+ */
+class FECACommand_GetSlotProperty : public IECACommand
+{
+public:
+	virtual FString GetName() const override { return TEXT("get_slot_property"); }
+	virtual FString GetDescription() const override { return TEXT("Read layout properties from a widget's slot (padding, alignment, size_rule, size_value). size_value is the fill weight when size_rule is Fill. When size_rule is Auto, size_value reflects the previously configured value and should be ignored."); }
+	virtual FString GetCategory() const override { return TEXT("UMG"); }
+
+	virtual TArray<FECACommandParam> GetParameters() const override
+	{
+		return {
+			{ TEXT("widget_path"), TEXT("string"), TEXT("Path to the Widget Blueprint"), true },
+			{ TEXT("widget_name"), TEXT("string"), TEXT("Name of the child widget whose slot to read"), true },
+		};
+	}
+
+	virtual FECACommandResult Execute(const TSharedPtr<FJsonObject>& Params) override;
+};
+
+/**
+ * Create a new empty UWidgetAnimation on a Widget Blueprint.
+ * The animation appears in the Animations panel and can be referenced in Blueprint graphs.
+ */
+class FECACommand_CreateWidgetAnimation : public IECACommand
+{
+public:
+	virtual FString GetName() const override { return TEXT("create_widget_animation"); }
+	virtual FString GetDescription() const override { return TEXT("Create a new empty animation on a Widget Blueprint (e.g. Anim_FadeIn). The animation can then be referenced in Blueprint function graphs."); }
+	virtual FString GetCategory() const override { return TEXT("UMG"); }
+
+	virtual TArray<FECACommandParam> GetParameters() const override
+	{
+		return {
+			{ TEXT("widget_path"), TEXT("string"), TEXT("Path to the Widget Blueprint"), true },
+			{ TEXT("animation_name"), TEXT("string"), TEXT("Name for the new animation (e.g. Anim_FadeIn)"), true },
+		};
+	}
+
+	virtual FECACommandResult Execute(const TSharedPtr<FJsonObject>& Params) override;
+};
