ECABridge: Fixed compile error w/ missing include

#Synced-CL 52144611

diff --git a/Engine/Plugins/Experimental/ECABridge/Source/ECABridge/Private/Commands/ECAToolsetRegistryCommands.cpp b/Engine/Plugins/Experimental/ECABridge/Source/ECABridge/Private/Commands/ECAToolsetRegistryCommands.cpp
--- a/Engine/Plugins/Experimental/ECABridge/Source/ECABridge/Private/Commands/ECAToolsetRegistryCommands.cpp
+++ b/Engine/Plugins/Experimental/ECABridge/Source/ECABridge/Private/Commands/ECAToolsetRegistryCommands.cpp
@@ -1,588 +1,277 @@
 // Copyright Epic Games, Inc. All Rights Reserved.
 
 #include "Commands/ECACommand.h"
+#include "Misc/CoreDelegates.h"
 #include "Dom/JsonObject.h"
 #include "Dom/JsonValue.h"
 #include "Serialization/JsonSerializer.h"
 #include "Serialization/JsonWriter.h"
 #include "Serialization/JsonReader.h"
-#include "Misc/FileHelper.h"
-#include "Misc/Paths.h"
-#include "HAL/PlatformFileManager.h"
-#include "HAL/PlatformProcess.h"
+
+#if WITH_TOOLSET_REGISTRY
+#include "ECABridgeToolset.h"
+#include "ToolsetRegistry/ToolsetRegistrySubsystem.h"
+#include "ToolsetRegistry/ToolsetRegistry.h"
+#include "ToolsetRegistry/Toolset.h"
+#endif
 
 //------------------------------------------------------------------------------
-// Toolset Registry State
-// 
-// The ECABridge C++ side maintains a cached view of the Python-side function
-// registry. This cache is populated by ECA calling store_registered_functions
-// (which syncs from the Python registry), or by the bridge querying directly.
+// ECABridge <-> ToolsetRegistry Integration
+//
+// This file provides two-way integration between ECABridge and the
+// ToolsetRegistry plugin:
 //
-// For function calls, the bridge queues a request that ECA picks up via the
-// event system and returns results through the normal MCP response flow.
+// 1. ECABridge commands exposed TO the ToolsetRegistry:
+//    - Registers all ECABridge commands as a FToolset so other consumers
+//      (Copilot, Blueprint agents, etc.) can discover and call them.
+//
+// 2. ToolsetRegistry tools accessible FROM ECABridge:
+//    - MCP commands that let the AI list and execute tools from any
+//      registered toolset (actor, asset, blueprint, material, etc.)
 //------------------------------------------------------------------------------
 
-namespace ToolsetRegistryCache
-{
-	struct FRegisteredFunctionInfo
-	{
-		FString Name;
-		FString Description;
-		FString Category;
-		FString Scope;
-		FString ReturnType;
-		FString SourceCode;
-		TArray<TSharedPtr<FJsonObject>> Parameters;
-	};
-
-	static TMap<FString, FRegisteredFunctionInfo> CachedFunctions;
-	static FCriticalSection CacheLock;
-
-	void UpdateCache(const TArray<TSharedPtr<FJsonValue>>& Functions)
-	{
-		FScopeLock Lock(&CacheLock);
-		CachedFunctions.Empty();
+#if WITH_TOOLSET_REGISTRY
 
-		for (const TSharedPtr<FJsonValue>& FuncValue : Functions)
-		{
-			const TSharedPtr<FJsonObject>& FuncObj = FuncValue->AsObject();
-			if (!FuncObj.IsValid())
-			{
-				continue;
-			}
-
-			FRegisteredFunctionInfo Info;
-			FuncObj->TryGetStringField(TEXT("name"), Info.Name);
-			FuncObj->TryGetStringField(TEXT("description"), Info.Description);
-			FuncObj->TryGetStringField(TEXT("category"), Info.Category);
-			FuncObj->TryGetStringField(TEXT("scope"), Info.Scope);
-			FuncObj->TryGetStringField(TEXT("return_type"), Info.ReturnType);
-			FuncObj->TryGetStringField(TEXT("source"), Info.SourceCode);
-
-			const TArray<TSharedPtr<FJsonValue>>* ParamsArray;
-			if (FuncObj->TryGetArrayField(TEXT("parameters"), ParamsArray))
-			{
-				for (const TSharedPtr<FJsonValue>& ParamVal : *ParamsArray)
-				{
-					if (ParamVal->Type == EJson::Object)
-					{
-						Info.Parameters.Add(ParamVal->AsObject());
-					}
-				}
-			}
+// FECABridgeToolset is declared in ECABridgeToolset.h and defined in ECABridgeToolset.cpp
 
-			if (!Info.Name.IsEmpty())
-			{
-				CachedFunctions.Add(Info.Name, MoveTemp(Info));
-			}
-		}
+// Singleton for the ECABridge toolset instance
+static TSharedPtr<FECABridgeToolset> GECABridgeToolset;
 
-		UE_LOG(LogTemp, Log, TEXT("[ECABridge] Toolset registry cache updated with %d functions"), CachedFunctions.Num());
+// Registration/unregistration helpers
+static void RegisterECABridgeWithToolsetRegistry()
+{
+	auto Result = UToolsetRegistrySubsystem::Get(TEXT("ECABridge toolset registration"));
+	if (Result.HasError())
+	{
+		UE_LOG(LogTemp, Log, TEXT("[ECABridge] ToolsetRegistry subsystem not available yet, will retry later"));
+		return;
 	}
 
-	TArray<FRegisteredFunctionInfo> GetAll(const FString& ScopeFilter = TEXT(""), const FString& CategoryFilter = TEXT(""))
+	UToolsetRegistrySubsystem* Subsystem = Result.GetValue();
+	if (!Subsystem)
 	{
-		FScopeLock Lock(&CacheLock);
-		TArray<FRegisteredFunctionInfo> Result;
-
-		for (const auto& Pair : CachedFunctions)
-		{
-			const FRegisteredFunctionInfo& Info = Pair.Value;
-
-			if (!ScopeFilter.IsEmpty() && Info.Scope != ScopeFilter)
-			{
-				continue;
-			}
-			if (!CategoryFilter.IsEmpty() && Info.Category != CategoryFilter)
-			{
-				continue;
-			}
-
-			Result.Add(Info);
-		}
+		return;
+	}
 
-		return Result;
+	if (!GECABridgeToolset.IsValid())
+	{
+		GECABridgeToolset = MakeShared<FECABridgeToolset>();
 	}
 
-	const FRegisteredFunctionInfo* Find(const FString& Name)
+	if (Subsystem->ToolsetRegistry.RegisterToolset(GECABridgeToolset))
 	{
-		FScopeLock Lock(&CacheLock);
-		return CachedFunctions.Find(Name);
+		UE_LOG(LogTemp, Log, TEXT("[ECABridge] Registered %d commands with ToolsetRegistry as 'ECABridge' toolset"),
+			FECACommandRegistry::Get().GetAllCommands().Num());
 	}
 }
 
-//------------------------------------------------------------------------------
-// store_registered_functions - Internal command for ECA to sync registry state
-// 
-// ECA calls this after loading/modifying the function registry to keep the
-// bridge's cache up to date. This is the primary sync mechanism.
-//------------------------------------------------------------------------------
-
-class FECACommand_StoreRegisteredFunctions : public IECACommand
+static void UnregisterECABridgeFromToolsetRegistry()
 {
-public:
-	virtual FString GetName() const override { return TEXT("store_registered_functions"); }
-	virtual FString GetDescription() const override 
-	{ 
-		return TEXT("Internal: Sync the registered functions cache from the ECA Python side. Called automatically by ECA when the registry changes."); 
-	}
-	virtual FString GetCategory() const override { return TEXT("ToolsetRegistry"); }
-	
-	virtual TArray<FECACommandParam> GetParameters() const override
+	if (!GECABridgeToolset.IsValid())
 	{
-		return {
-			{ TEXT("functions"), TEXT("array"), TEXT("Array of registered function info objects"), true }
-		};
+		return;
 	}
 
-	virtual FECACommandResult Execute(const TSharedPtr<FJsonObject>& Params) override
+	auto Result = UToolsetRegistrySubsystem::Get();
+	if (!Result.HasError())
 	{
-		const TArray<TSharedPtr<FJsonValue>>* FunctionsArray;
-		if (!Params->TryGetArrayField(TEXT("functions"), FunctionsArray))
+		UToolsetRegistrySubsystem* Subsystem = Result.GetValue();
+		if (Subsystem)
 		{
-			return FECACommandResult::Error(TEXT("Missing required parameter: functions (array)"));
+			Subsystem->ToolsetRegistry.UnregisterToolset(GECABridgeToolset);
 		}
+	}
 
-		ToolsetRegistryCache::UpdateCache(*FunctionsArray);
+	GECABridgeToolset.Reset();
+}
 
-		TSharedPtr<FJsonObject> Result = MakeResult();
-		Result->SetBoolField(TEXT("success"), true);
-		Result->SetNumberField(TEXT("cached_count"), FunctionsArray->Num());
+#endif // WITH_TOOLSET_REGISTRY
 
-		return FECACommandResult::Success(Result);
-	}
-};
+//------------------------------------------------------------------------------
+// MCP Commands for ToolsetRegistry access
+//------------------------------------------------------------------------------
 
 //------------------------------------------------------------------------------
-// list_registered_functions
+// toolset_list - List all registered toolsets and their tools
 //------------------------------------------------------------------------------
 
-class FECACommand_ListRegisteredFunctions : public IECACommand
+class FECACommand_ToolsetList : public IECACommand
 {
 public:
-	virtual FString GetName() const override { return TEXT("list_registered_functions"); }
-	virtual FString GetDescription() const override { return TEXT("List all registered functions from the toolset registry cache, with optional scope/category filtering."); }
+	virtual FString GetName() const override { return TEXT("toolset_list"); }
+	virtual FString GetDescription() const override
+	{
+		return TEXT("List all registered toolsets in the ToolsetRegistry. Returns toolset names, versions, and their tool schemas. Use this to discover available tools from all registered toolsets (actor, asset, blueprint, material, etc.).");
+	}
 	virtual FString GetCategory() const override { return TEXT("ToolsetRegistry"); }
 
 	virtual TArray<FECACommandParam> GetParameters() const override
 	{
 		return {
-			{ TEXT("scope"), TEXT("string"), TEXT("Filter by scope (e.g. 'project', 'global')"), false },
-			{ TEXT("category"), TEXT("string"), TEXT("Filter by category"), false }
+			{ TEXT("toolset_name"), TEXT("string"), TEXT("If specified, only return the schema for this toolset"), false }
 		};
 	}
 
-	virtual FECACommandResult Execute(const TSharedPtr<FJsonObject>& Params) override;
-};
-
-FECACommandResult FECACommand_ListRegisteredFunctions::Execute(const TSharedPtr<FJsonObject>& Params)
-{
-	FString ScopeFilter;
-	GetStringParam(Params, TEXT("scope"), ScopeFilter, false);
-
-	FString CategoryFilter;
-	GetStringParam(Params, TEXT("category"), CategoryFilter, false);
-
-	TArray<ToolsetRegistryCache::FRegisteredFunctionInfo> Functions = 
-		ToolsetRegistryCache::GetAll(ScopeFilter, CategoryFilter);
-
-	TSharedPtr<FJsonObject> Result = MakeResult();
-	
-	TArray<TSharedPtr<FJsonValue>> FunctionsArray;
-	for (const auto& Info : Functions)
+	virtual FECACommandResult Execute(const TSharedPtr<FJsonObject>& Params) override
 	{
-		TSharedPtr<FJsonObject> FuncObj = MakeShared<FJsonObject>();
-		FuncObj->SetStringField(TEXT("name"), Info.Name);
-		FuncObj->SetStringField(TEXT("description"), Info.Description);
-		FuncObj->SetStringField(TEXT("category"), Info.Category);
-		FuncObj->SetStringField(TEXT("scope"), Info.Scope);
-		FuncObj->SetStringField(TEXT("return_type"), Info.ReturnType);
-		
-		// Include parameter info
-		TArray<TSharedPtr<FJsonValue>> ParamsArray;
-		for (const TSharedPtr<FJsonObject>& ParamObj : Info.Parameters)
+#if WITH_TOOLSET_REGISTRY
+		auto SubsystemResult = UToolsetRegistrySubsystem::Get(TEXT("toolset_list"));
+		if (SubsystemResult.HasError())
 		{
-			ParamsArray.Add(MakeShared<FJsonValueObject>(ParamObj));
+			return FECACommandResult::Error(FString::Printf(TEXT("ToolsetRegistry not available: %s"), *SubsystemResult.GetError()));
 		}
-		FuncObj->SetArrayField(TEXT("parameters"), ParamsArray);
-		
-		FunctionsArray.Add(MakeShared<FJsonValueObject>(FuncObj));
-	}
 
-	Result->SetArrayField(TEXT("functions"), FunctionsArray);
-	Result->SetNumberField(TEXT("count"), Functions.Num());
+		UToolsetRegistrySubsystem* Subsystem = SubsystemResult.GetValue();
 
-	return FECACommandResult::Success(Result);
-}
+		FString ToolsetName;
+		GetStringParam(Params, TEXT("toolset_name"), ToolsetName, false);
 
-//------------------------------------------------------------------------------
-// call_registered_function
-//
-// This is the key integration point. When Unreal-side code (or the AI via
-// ECABridge) needs to call a registered Python function, this command
-// queues a "call_function" event. ECA's event polling picks it up,
-// executes the function, and pushes the result back via a separate
-// "function_result" event or direct MCP response.
-//
-// For the synchronous MCP call path (AI calls tools/call with
-// call_registered_function), we serialize the request as JSON and ECA
-// executes it inline via meta_tool, returning the result directly.
-//------------------------------------------------------------------------------
-
-class FECACommand_CallRegisteredFunction : public IECACommand
-{
-public:
-	virtual FString GetName() const override { return TEXT("call_registered_function"); }
-	virtual FString GetDescription() const override { return TEXT("Call a registered Python function by name. The request is forwarded to ECA for execution."); }
-	virtual FString GetCategory() const override { return TEXT("ToolsetRegistry"); }
-
-	virtual TArray<FECACommandParam> GetParameters() const override
-	{
-		return {
-			{ TEXT("function_name"), TEXT("string"), TEXT("Name of the registered function to call"), true },
-			{ TEXT("arguments"), TEXT("object"), TEXT("Arguments to pass to the function"), false }
-		};
-	}
-
-	virtual FECACommandResult Execute(const TSharedPtr<FJsonObject>& Params) override;
-};
-
-FECACommandResult FECACommand_CallRegisteredFunction::Execute(const TSharedPtr<FJsonObject>& Params)
-{
-	FString FunctionName;
-	if (!GetStringParam(Params, TEXT("function_name"), FunctionName))
-	{
-		return FECACommandResult::Error(TEXT("Missing required parameter: function_name"));
-	}
+		TSharedPtr<FJsonObject> Result = MakeResult();
 
-	// Verify the function exists in our cache
-	const ToolsetRegistryCache::FRegisteredFunctionInfo* FuncInfo = ToolsetRegistryCache::Find(FunctionName);
-	if (!FuncInfo)
-	{
-		return FECACommandResult::Error(FString::Printf(
-			TEXT("Function '%s' not found in toolset registry cache. Call store_registered_functions first to sync the cache."),
-			*FunctionName));
-	}
+		if (!ToolsetName.IsEmpty())
+		{
+			// Find a specific toolset
+			FString ErrorMsg;
+			TSharedPtr<UE::ToolsetRegistry::FToolset> Toolset = Subsystem->ToolsetRegistry.Find(ToolsetName, true, &ErrorMsg);
+			if (!Toolset.IsValid())
+			{
+				return FECACommandResult::Error(FString::Printf(TEXT("Toolset '%s' not found: %s"), *ToolsetName, *ErrorMsg));
+			}
 
-	// Build the call request
-	TSharedPtr<FJsonObject> Arguments = MakeShared<FJsonObject>();
-	const TSharedPtr<FJsonObject>* ArgsPtr;
-	if (Params->TryGetObjectField(TEXT("arguments"), ArgsPtr))
-	{
-		Arguments = *ArgsPtr;
-	}
+			Result->SetStringField(TEXT("schema"), Toolset->GetToolsetJsonSchema());
+			Result->SetStringField(TEXT("toolset_name"), Toolset->GetToolsetName());
+			Result->SetStringField(TEXT("toolset_version"), Toolset->GetToolsetVersion());
+		}
+		else
+		{
+			// Return all toolset schemas
+			Result->SetStringField(TEXT("schemas"), Subsystem->ToolsetRegistry.GetToolsetJsonSchemas());
+		}
 
-	// Build a meta_tool code snippet that calls the registered function
-	// This will be executed by ECA when it processes the MCP response
-	FString ArgsJson;
-	{
-		TSharedRef<FJsonObject> ArgsRef = Arguments.ToSharedRef();
-		TSharedRef<TJsonWriter<>> Writer = TJsonWriterFactory<>::Create(&ArgsJson);
-		FJsonSerializer::Serialize(ArgsRef, Writer);
+		return FECACommandResult::Success(Result);
+#else
+		return FECACommandResult::Error(TEXT("ToolsetRegistry plugin not available. Compile with WITH_TOOLSET_REGISTRY=1."));
+#endif
 	}
-
-	TSharedPtr<FJsonObject> Result = MakeResult();
-	Result->SetStringField(TEXT("function_name"), FunctionName);
-	Result->SetObjectField(TEXT("arguments"), Arguments);
-	Result->SetStringField(TEXT("description"), FuncInfo->Description);
-	Result->SetStringField(TEXT("category"), FuncInfo->Category);
-	Result->SetStringField(TEXT("status"), TEXT("call_requested"));
-	Result->SetStringField(TEXT("instruction"), 
-		FString::Printf(TEXT("Execute this registered function via meta_tool: tools.%s(**%s)"), 
-			*FunctionName, *ArgsJson));
-
-	return FECACommandResult::Success(Result);
-}
+};
 
 //------------------------------------------------------------------------------
-// get_registered_function_info
+// toolset_call - Execute a tool from a registered toolset
 //------------------------------------------------------------------------------
 
-class FECACommand_GetRegisteredFunctionInfo : public IECACommand
+class FECACommand_ToolsetCall : public IECACommand
 {
 public:
-	virtual FString GetName() const override { return TEXT("get_registered_function_info"); }
-	virtual FString GetDescription() const override { return TEXT("Get detailed info about a specific registered function including source code and parameters."); }
-	virtual FString GetCategory() const override { return TEXT("ToolsetRegistry"); }
-
-	virtual TArray<FECACommandParam> GetParameters() const override
+	virtual FString GetName() const override { return TEXT("toolset_call"); }
+	virtual FString GetDescription() const override
 	{
-		return {
-			{ TEXT("function_name"), TEXT("string"), TEXT("Name of the registered function"), true }
-		};
+		return TEXT("Execute a tool from a registered toolset. Format: 'toolset_name.tool_name' or just 'tool_name' if unambiguous. Arguments are passed as JSON.");
 	}
-
-	virtual FECACommandResult Execute(const TSharedPtr<FJsonObject>& Params) override;
-};
-
-FECACommandResult FECACommand_GetRegisteredFunctionInfo::Execute(const TSharedPtr<FJsonObject>& Params)
-{
-	FString FunctionName;
-	if (!GetStringParam(Params, TEXT("function_name"), FunctionName))
-	{
-		return FECACommandResult::Error(TEXT("Missing required parameter: function_name"));
-	}
-
-	const ToolsetRegistryCache::FRegisteredFunctionInfo* FuncInfo = ToolsetRegistryCache::Find(FunctionName);
-	if (!FuncInfo)
-	{
-		return FECACommandResult::Error(FString::Printf(
-			TEXT("Function '%s' not found in toolset registry cache."), *FunctionName));
-	}
-
-	TSharedPtr<FJsonObject> Result = MakeResult();
-	Result->SetStringField(TEXT("name"), FuncInfo->Name);
-	Result->SetStringField(TEXT("description"), FuncInfo->Description);
-	Result->SetStringField(TEXT("category"), FuncInfo->Category);
-	Result->SetStringField(TEXT("scope"), FuncInfo->Scope);
-	Result->SetStringField(TEXT("return_type"), FuncInfo->ReturnType);
-
-	if (!FuncInfo->SourceCode.IsEmpty())
-	{
-		Result->SetStringField(TEXT("source_code"), FuncInfo->SourceCode);
-	}
-
-	TArray<TSharedPtr<FJsonValue>> ParamsArray;
-	for (const TSharedPtr<FJsonObject>& ParamObj : FuncInfo->Parameters)
-	{
-		ParamsArray.Add(MakeShared<FJsonValueObject>(ParamObj));
-	}
-	Result->SetArrayField(TEXT("parameters"), ParamsArray);
-
-	return FECACommandResult::Success(Result);
-}
-
-//------------------------------------------------------------------------------
-// register_function
-//
-// Queues a function registration request. ECA processes this by writing
-// the function to .eca/functions/ and loading it into the Python registry.
-// The bridge cache is then updated via store_registered_functions.
-//------------------------------------------------------------------------------
-
-class FECACommand_RegisterFunction : public IECACommand
-{
-public:
-	virtual FString GetName() const override { return TEXT("register_function"); }
-	virtual FString GetDescription() const override { return TEXT("Request registration of a new Python function. ECA will write and load the function."); }
 	virtual FString GetCategory() const override { return TEXT("ToolsetRegistry"); }
 
 	virtual TArray<FECACommandParam> GetParameters() const override
 	{
 		return {
-			{ TEXT("function_code"), TEXT("string"), TEXT("Python source code for the function"), true },
-			{ TEXT("function_name"), TEXT("string"), TEXT("Name for the function"), true },
-			{ TEXT("scope"), TEXT("string"), TEXT("Scope: 'project' or 'global'"), false, TEXT("project") },
-			{ TEXT("category"), TEXT("string"), TEXT("Category for organization"), false, TEXT("User-Defined") }
+			{ TEXT("tool"), TEXT("string"), TEXT("Tool identifier: 'toolset.tool_name' or just 'tool_name'"), true },
+			{ TEXT("arguments"), TEXT("string"), TEXT("JSON arguments for the tool"), false, TEXT("{}") }
 		};
 	}
 
-	virtual FECACommandResult Execute(const TSharedPtr<FJsonObject>& Params) override;
-};
-
-FECACommandResult FECACommand_RegisterFunction::Execute(const TSharedPtr<FJsonObject>& Params)
-{
-	FString FunctionCode;
-	if (!GetStringParam(Params, TEXT("function_code"), FunctionCode))
+	virtual FECACommandResult Execute(const TSharedPtr<FJsonObject>& Params) override
 	{
-		return FECACommandResult::Error(TEXT("Missing required parameter: function_code"));
-	}
+#if WITH_TOOLSET_REGISTRY
+		FString ToolIdentifier;
+		if (!GetStringParam(Params, TEXT("tool"), ToolIdentifier))
+		{
+			return FECACommandResult::Error(TEXT("Missing required parameter: tool"));
+		}
 
-	FString FunctionName;
-	if (!GetStringParam(Params, TEXT("function_name"), FunctionName))
-	{
-		return FECACommandResult::Error(TEXT("Missing required parameter: function_name"));
-	}
+		FString Arguments = TEXT("{}");
+		GetStringParam(Params, TEXT("arguments"), Arguments, false);
 
-	FString Scope = TEXT("project");
-	GetStringParam(Params, TEXT("scope"), Scope, false);
-
-	FString Category = TEXT("User-Defined");
-	GetStringParam(Params, TEXT("category"), Category, false);
-
-	// Build the result with the registration request info
-	// ECA will execute this via its own register_function tool
-	TSharedPtr<FJsonObject> Result = MakeResult();
-	Result->SetStringField(TEXT("function_name"), FunctionName);
-	Result->SetStringField(TEXT("function_code"), FunctionCode);
-	Result->SetStringField(TEXT("scope"), Scope);
-	Result->SetStringField(TEXT("category"), Category);
-	Result->SetStringField(TEXT("status"), TEXT("registration_requested"));
-	Result->SetStringField(TEXT("instruction"),
-		FString::Printf(TEXT("Register this function via meta_tool: tools.register_function(function_code=%s, function_name='%s', scope='%s', category='%s')"),
-			*FunctionCode, *FunctionName, *Scope, *Category));
-
-	return FECACommandResult::Success(Result);
-}
+		// Also accept arguments as an object (not just string)
+		const TSharedPtr<FJsonObject>* ArgsObjPtr;
+		if (Params->TryGetObjectField(TEXT("arguments"), ArgsObjPtr))
+		{
+			FString ArgsJson;
+			TSharedRef<TJsonWriter<>> Writer = TJsonWriterFactory<>::Create(&ArgsJson);
+			FJsonSerializer::Serialize((*ArgsObjPtr).ToSharedRef(), Writer);
+			Arguments = ArgsJson;
+		}
 
-//------------------------------------------------------------------------------
-// unregister_function
-//------------------------------------------------------------------------------
+		auto SubsystemResult = UToolsetRegistrySubsystem::Get(TEXT("toolset_call"));
+		if (SubsystemResult.HasError())
+		{
+			return FECACommandResult::Error(FString::Printf(TEXT("ToolsetRegistry not available: %s"), *SubsystemResult.GetError()));
+		}
 
-class FECACommand_UnregisterFunction : public IECACommand
-{
-public:
-	virtual FString GetName() const override { return TEXT("unregister_function"); }
-	virtual FString GetDescription() const override { return TEXT("Request unregistration of a registered Python function."); }
-	virtual FString GetCategory() const override { return TEXT("ToolsetRegistry"); }
+		UToolsetRegistrySubsystem* Subsystem = SubsystemResult.GetValue();
 
-	virtual TArray<FECACommandParam> GetParameters() const override
-	{
-		return {
-			{ TEXT("function_name"), TEXT("string"), TEXT("Name of the function to unregister"), true },
-			{ TEXT("scope"), TEXT("string"), TEXT("Scope to search in"), false }
-		};
-	}
+		// Parse tool descriptor (handles "toolset.tool" format)
+		auto DescriptorResult = UE::ToolsetRegistry::FToolDescriptor::FromString(ToolIdentifier, true);
+		if (DescriptorResult.HasError())
+		{
+			return FECACommandResult::Error(FString::Printf(TEXT("Invalid tool identifier '%s': %s"), *ToolIdentifier, *DescriptorResult.GetError()));
+		}
 
-	virtual FECACommandResult Execute(const TSharedPtr<FJsonObject>& Params) override;
-};
+		// Execute the tool - this returns a future, we need to wait for it
+		TFuture<TValueOrError<FString, FString>> Future =
+			Subsystem->ToolsetRegistry.ExecuteTool(DescriptorResult.GetValue(), Arguments);
 
-FECACommandResult FECACommand_UnregisterFunction::Execute(const TSharedPtr<FJsonObject>& Params)
-{
-	FString FunctionName;
-	if (!GetStringParam(Params, TEXT("function_name"), FunctionName))
-	{
-		return FECACommandResult::Error(TEXT("Missing required parameter: function_name"));
-	}
+		// Block and wait for the result (MCP is synchronous request/response)
+		Future.Wait();
+		TValueOrError<FString, FString> ToolResult = Future.Get();
 
-	FString Scope;
-	GetStringParam(Params, TEXT("scope"), Scope, false);
+		if (ToolResult.HasError())
+		{
+			return FECACommandResult::Error(ToolResult.GetError());
+		}
 
-	TSharedPtr<FJsonObject> Result = MakeResult();
-	Result->SetStringField(TEXT("function_name"), FunctionName);
-	if (!Scope.IsEmpty())
-	{
-		Result->SetStringField(TEXT("scope"), Scope);
-	}
-	Result->SetStringField(TEXT("status"), TEXT("unregistration_requested"));
-	Result->SetStringField(TEXT("instruction"),
-		FString::Printf(TEXT("Unregister this function via meta_tool: tools.unregister_function(name='%s'%s)"),
-			*FunctionName,
-			Scope.IsEmpty() ? TEXT("") : *FString::Printf(TEXT(", scope='%s'"), *Scope)));
+		// Parse the result JSON string back into an object
+		FString ResultJson = ToolResult.GetValue();
+		TSharedPtr<FJsonObject> ResultObj;
+		TSharedRef<TJsonReader<>> Reader = TJsonReaderFactory<>::Create(ResultJson);
+		if (FJsonSerializer::Deserialize(Reader, ResultObj) && ResultObj.IsValid())
+		{
+			return FECACommandResult::Success(ResultObj);
+		}
 
-	return FECACommandResult::Success(Result);
-}
+		// If not parseable as JSON object, wrap as string
+		TSharedPtr<FJsonObject> Result = MakeResult();
+		Result->SetStringField(TEXT("result"), ResultJson);
+		return FECACommandResult::Success(Result);
+#else
+		return FECACommandResult::Error(TEXT("ToolsetRegistry plugin not available. Compile with WITH_TOOLSET_REGISTRY=1."));
+#endif
+	}
+};
 
 //------------------------------------------------------------------------------
-// get_all_tools
+// Register commands and auto-register with ToolsetRegistry
 //------------------------------------------------------------------------------
 
-class FECACommand_GetAllTools : public IECACommand
-{
-public:
-	virtual FString GetName() const override { return TEXT("get_all_tools"); }
-	virtual FString GetDescription() const override { return TEXT("Get a combined list of all available tools: ECABridge commands and registered Python functions."); }
-	virtual FString GetCategory() const override { return TEXT("ToolsetRegistry"); }
-
-	virtual TArray<FECACommandParam> GetParameters() const override
-	{
-		return {
-			{ TEXT("include_bridge_commands"), TEXT("boolean"), TEXT("Include ECABridge commands"), false, TEXT("true") },
-			{ TEXT("include_registered_functions"), TEXT("boolean"), TEXT("Include registered Python functions"), false, TEXT("true") },
-			{ TEXT("category_filter"), TEXT("string"), TEXT("Filter by category"), false }
-		};
-	}
-
-	virtual FECACommandResult Execute(const TSharedPtr<FJsonObject>& Params) override;
-};
+REGISTER_ECA_COMMAND(FECACommand_ToolsetList)
+REGISTER_ECA_COMMAND(FECACommand_ToolsetCall)
 
-FECACommandResult FECACommand_GetAllTools::Execute(const TSharedPtr<FJsonObject>& Params)
+// Auto-register ECABridge as a toolset when the engine initializes
+#if WITH_TOOLSET_REGISTRY
+struct FECAToolsetRegistryAutoRegister
 {
-	bool bIncludeBridgeCommands = true;
-	GetBoolParam(Params, TEXT("include_bridge_commands"), bIncludeBridgeCommands, false);
-
-	bool bIncludeRegisteredFunctions = true;
-	GetBoolParam(Params, TEXT("include_registered_functions"), bIncludeRegisteredFunctions, false);
-
-	FString CategoryFilter;
-	GetStringParam(Params, TEXT("category_filter"), CategoryFilter, false);
-
-	TSharedPtr<FJsonObject> Result = MakeResult();
-	TArray<TSharedPtr<FJsonValue>> AllTools;
-	int32 BridgeCount = 0;
-	int32 RegisteredCount = 0;
-
-	// Add ECABridge commands
-	if (bIncludeBridgeCommands)
+	FECAToolsetRegistryAutoRegister()
 	{
-		for (const TSharedPtr<IECACommand>& Command : FECACommandRegistry::Get().GetAllCommands())
+		// Defer registration until the editor is fully initialized
+		// (ToolsetRegistrySubsystem may not exist yet at static init time)
+		FCoreDelegates::GetOnPostEngineInit().AddLambda([]()
 		{
-			if (!CategoryFilter.IsEmpty() && Command->GetCategory() != CategoryFilter)
-			{
-				continue;
-			}
-
-			TSharedPtr<FJsonObject> ToolObj = MakeShared<FJsonObject>();
-			ToolObj->SetStringField(TEXT("name"), Command->GetName());
-			ToolObj->SetStringField(TEXT("description"), Command->GetDescription());
-			ToolObj->SetStringField(TEXT("category"), Command->GetCategory());
-			ToolObj->SetStringField(TEXT("source"), TEXT("eca_bridge"));
-
-			TArray<TSharedPtr<FJsonValue>> ParamsArray;
-			for (const FECACommandParam& Param : Command->GetParameters())
-			{
-				TSharedPtr<FJsonObject> ParamObj = MakeShared<FJsonObject>();
-				ParamObj->SetStringField(TEXT("name"), Param.Name);
-				ParamObj->SetStringField(TEXT("type"), Param.Type);
-				ParamObj->SetStringField(TEXT("description"), Param.Description);
-				ParamObj->SetBoolField(TEXT("required"), Param.bRequired);
-				if (!Param.DefaultValue.IsEmpty())
-				{
-					ParamObj->SetStringField(TEXT("default"), Param.DefaultValue);
-				}
-				ParamsArray.Add(MakeShared<FJsonValueObject>(ParamObj));
-			}
-			ToolObj->SetArrayField(TEXT("parameters"), ParamsArray);
-
-			AllTools.Add(MakeShared<FJsonValueObject>(ToolObj));
-			BridgeCount++;
-		}
+			RegisterECABridgeWithToolsetRegistry();
+		});
 	}
 
-	// Add registered functions from cache
-	if (bIncludeRegisteredFunctions)
+	~FECAToolsetRegistryAutoRegister()
 	{
-		TArray<ToolsetRegistryCache::FRegisteredFunctionInfo> Functions =
-			ToolsetRegistryCache::GetAll(TEXT(""), CategoryFilter);
-
-		for (const auto& Info : Functions)
-		{
-			TSharedPtr<FJsonObject> ToolObj = MakeShared<FJsonObject>();
-			ToolObj->SetStringField(TEXT("name"), Info.Name);
-			ToolObj->SetStringField(TEXT("description"), Info.Description);
-			ToolObj->SetStringField(TEXT("category"), Info.Category);
-			ToolObj->SetStringField(TEXT("source"), TEXT("toolset_registry"));
-			ToolObj->SetStringField(TEXT("scope"), Info.Scope);
-			ToolObj->SetStringField(TEXT("return_type"), Info.ReturnType);
-
-			TArray<TSharedPtr<FJsonValue>> ParamsArray;
-			for (const TSharedPtr<FJsonObject>& ParamObj : Info.Parameters)
-			{
-				ParamsArray.Add(MakeShared<FJsonValueObject>(ParamObj));
-			}
-			ToolObj->SetArrayField(TEXT("parameters"), ParamsArray);
-
-			AllTools.Add(MakeShared<FJsonValueObject>(ToolObj));
-			RegisteredCount++;
-		}
+		UnregisterECABridgeFromToolsetRegistry();
 	}
+};
 
-	Result->SetArrayField(TEXT("tools"), AllTools);
-	Result->SetNumberField(TEXT("total_count"), AllTools.Num());
-	Result->SetNumberField(TEXT("bridge_command_count"), BridgeCount);
-	Result->SetNumberField(TEXT("registered_function_count"), RegisteredCount);
-
-	return FECACommandResult::Success(Result);
-}
-
-//------------------------------------------------------------------------------
-// Register all commands
-//------------------------------------------------------------------------------
-
-REGISTER_ECA_COMMAND(FECACommand_StoreRegisteredFunctions)
-REGISTER_ECA_COMMAND(FECACommand_ListRegisteredFunctions)
-REGISTER_ECA_COMMAND(FECACommand_CallRegisteredFunction)
-REGISTER_ECA_COMMAND(FECACommand_GetRegisteredFunctionInfo)
-REGISTER_ECA_COMMAND(FECACommand_RegisterFunction)
-REGISTER_ECA_COMMAND(FECACommand_UnregisterFunction)
-REGISTER_ECA_COMMAND(FECACommand_GetAllTools)
\ No newline at end of file
+static FECAToolsetRegistryAutoRegister GToolsetRegistryAutoRegister;
+#endif
\ No newline at end of file
