ECABridge: BlueprintLisp improvement, accept string values

#Synced-CL 52144611

diff --git a/Engine/Plugins/Experimental/ECABridge/Source/ECABridge/Private/Commands/ECABlueprintLispCommands.cpp b/Engine/Plugins/Experimental/ECABridge/Source/ECABridge/Private/Commands/ECABlueprintLispCommands.cpp
--- a/Engine/Plugins/Experimental/ECABridge/Source/ECABridge/Private/Commands/ECABlueprintLispCommands.cpp
+++ b/Engine/Plugins/Experimental/ECABridge/Source/ECABridge/Private/Commands/ECABlueprintLispCommands.cpp
@@ -2899,6 +2899,10 @@ static UEdGraphPin* ResolveLispExpression(const FLispNodePtr& Expr, FLispToBPCon
 										SetNumericPinDefaultValue(Pin, NumVal);
 									}
 								}
+								else if (Expr->Get(i)->IsString())
+								{
+									Pin->DefaultValue = Expr->Get(i)->StringValue;
+								}
 								ArgIndex++;
 								break;
 							}
