[ECABridge] Blueprint function authoring tools

Adds MCP tools for creating and modifying Blueprint functions at runtime,
enabling AI to build and populate Blueprint logic without opening the editor
manually. Reviewers and users gain the ability to programmatically add,
remove, and rename functions and their parameters on any Blueprint asset.

New tools: add_blueprint_function, remove_blueprint_function,
add_blueprint_function_parameter, remove_blueprint_function_parameter,
rename_blueprint_function.
#rnx
#Synced-CL 52344065

diff --git a/Engine/Plugins/Experimental/ECABridge/Source/ECABridge/Private/Commands/ECABlueprintNodeCommands.cpp b/Engine/Plugins/Experimental/ECABridge/Source/ECABridge/Private/Commands/ECABlueprintNodeCommands.cpp
--- a/Engine/Plugins/Experimental/ECABridge/Source/ECABridge/Private/Commands/ECABlueprintNodeCommands.cpp
+++ b/Engine/Plugins/Experimental/ECABridge/Source/ECABridge/Private/Commands/ECABlueprintNodeCommands.cpp
@@ -23,6 +23,7 @@
 #include "K2Node_ComponentBoundEvent.h"
 #include "K2Node_MacroInstance.h"
 #include "K2Node_DynamicCast.h"
+#include "K2Node_ExecutionSequence.h"
 #include "K2Node_IfThenElse.h"
 #include "K2Node_Select.h"
 #include "K2Node_SwitchEnum.h"
@@ -629,10 +630,24 @@ FECACommandResult FECACommand_AddBlueprintFunctionNode::Execute(const TSharedPtr
 	
 	if (!TargetClass.IsEmpty())
 	{
-		UClass* SpecificClass = FindObject<UClass>(nullptr, *FString::Printf(TEXT("/Script/Engine.%s"), *TargetClass));
-		if (!SpecificClass)
+		// Try common module paths for the class
+		UClass* SpecificClass = nullptr;
+		for (const FString& ModulePath : TArray<FString>{
+			TEXT("/Script/UMG"),
+			TEXT("/Script/Engine"),
+			TEXT("/Script/CoreUObject"),
+			TEXT("/Script/SlateCore"),
+			TEXT("/Script/Slate"),
+		})
 		{
-			SpecificClass = FindObject<UClass>(nullptr, *FString::Printf(TEXT("/Script/CoreUObject.%s"), *TargetClass));
+			SpecificClass = FindObject<UClass>(nullptr, *FString::Printf(TEXT("%s.%s"), *ModulePath, *TargetClass));
+			if (!SpecificClass)
+			{
+				// Try with U prefix stripped
+				FString StrippedName = TargetClass.StartsWith(TEXT("U")) ? TargetClass.Mid(1) : TargetClass;
+				SpecificClass = FindObject<UClass>(nullptr, *FString::Printf(TEXT("%s.%s"), *ModulePath, *StrippedName));
+			}
+			if (SpecificClass) break;
 		}
 		// Fallback: search all loaded modules for the class by name
 		if (!SpecificClass)
@@ -1415,10 +1430,18 @@ FECACommandResult FECACommand_BatchEditBlueprintNodes::Execute(const TSharedPtr<
 			
 			if (!TargetClass.IsEmpty())
 			{
-				UClass* SpecificClass = FindObject<UClass>(nullptr, *FString::Printf(TEXT("/Script/Engine.%s"), *TargetClass));
-				if (!SpecificClass)
+				UClass* SpecificClass = nullptr;
+				for (const FString& ModulePath : TArray<FString>{
+					TEXT("/Script/UMG"), TEXT("/Script/Engine"),
+					TEXT("/Script/CoreUObject"), TEXT("/Script/SlateCore"), TEXT("/Script/Slate")})
 				{
-					SpecificClass = FindObject<UClass>(nullptr, *FString::Printf(TEXT("/Script/CoreUObject.%s"), *TargetClass));
+					SpecificClass = FindObject<UClass>(nullptr, *FString::Printf(TEXT("%s.%s"), *ModulePath, *TargetClass));
+					if (!SpecificClass)
+					{
+						FString Stripped = TargetClass.StartsWith(TEXT("U")) ? TargetClass.Mid(1) : TargetClass;
+						SpecificClass = FindObject<UClass>(nullptr, *FString::Printf(TEXT("%s.%s"), *ModulePath, *Stripped));
+					}
+					if (SpecificClass) break;
 				}
 				if (SpecificClass)
 				{
@@ -1457,6 +1480,28 @@ FECACommandResult FECACommand_BatchEditBlueprintNodes::Execute(const TSharedPtr<
 				ClassesToSearch.Add(USceneComponent::StaticClass());
 				ClassesToSearch.Add(UPrimitiveComponent::StaticClass());
 				ClassesToSearch.Add(Blueprint->ParentClass);
+
+				// Add target_class if specified — resolve by name across common module paths
+				if (!TargetClass.IsEmpty())
+				{
+					UClass* SpecificClass = nullptr;
+					for (const FString& ModulePath : TArray<FString>{
+						TEXT("/Script/UMG"), TEXT("/Script/Engine"),
+						TEXT("/Script/CoreUObject"), TEXT("/Script/SlateCore"), TEXT("/Script/Slate")})
+					{
+						SpecificClass = FindObject<UClass>(nullptr, *FString::Printf(TEXT("%s.%s"), *ModulePath, *TargetClass));
+						if (!SpecificClass)
+						{
+							FString Stripped = TargetClass.StartsWith(TEXT("U")) ? TargetClass.Mid(1) : TargetClass;
+							SpecificClass = FindObject<UClass>(nullptr, *FString::Printf(TEXT("%s.%s"), *ModulePath, *Stripped));
+						}
+						if (SpecificClass) break;
+					}
+					if (!SpecificClass)
+						SpecificClass = FindFirstObject<UClass>(*TargetClass, EFindFirstObjectOptions::NativeFirst);
+					if (SpecificClass)
+						ClassesToSearch.Add(SpecificClass);
+				}
 				
 				// Add Geometry Script library classes
 				static const TCHAR* GeometryScriptClassPaths[] = {
@@ -1748,9 +1793,24 @@ FECACommandResult FECACommand_BatchEditBlueprintNodes::Execute(const TSharedPtr<
 			// and PostPlacedNewNode has run; the finalization block skips ADP when Pins.Num() > 0.
 			CreatedNode = CreateWidgetNode;
 		}
+		else if (NodeType == TEXT("branch") || NodeType == TEXT("if") || NodeType == TEXT("ifthenelse"))
+		{
+			UK2Node_IfThenElse* BranchNode = NewObject<UK2Node_IfThenElse>(Graph);
+			CreatedNode = BranchNode;
+		}
+		else if (NodeType == TEXT("select"))
+		{
+			UK2Node_Select* SelectNode = NewObject<UK2Node_Select>(Graph);
+			CreatedNode = SelectNode;
+		}
+		else if (NodeType == TEXT("sequence"))
+		{
+			UK2Node_ExecutionSequence* SequenceNode = NewObject<UK2Node_ExecutionSequence>(Graph);
+			CreatedNode = SequenceNode;
+		}
 		else
 		{
-			Errors.Add(FString::Printf(TEXT("Node '%s': Unknown type '%s'. Valid types: event, custom_event, function, variable_get, variable_set, self, component, input_action, cast, macro, create_widget"), *TempId, *NodeType));
+			Errors.Add(FString::Printf(TEXT("Node '%s': Unknown type '%s'. Valid types: event, custom_event, function, variable_get, variable_set, self, component, input_action, cast, macro, create_widget, branch, select, sequence"), *TempId, *NodeType));
 			continue;
 		}
 		
@@ -1766,6 +1826,13 @@ FECACommandResult FECACommand_BatchEditBlueprintNodes::Execute(const TSharedPtr<
 				CreatedNode->AllocateDefaultPins();
 			}
 			EnsureBPNodeHasValidGuid(CreatedNode);
+			// Flow-control nodes (branch/select/sequence) need ReconstructNode after pin allocation
+			// to fully initialize their pin configuration (matches AddBlueprintFlowControlNode behaviour).
+			if (NodeType == TEXT("branch") || NodeType == TEXT("if") || NodeType == TEXT("ifthenelse") ||
+			    NodeType == TEXT("select") || NodeType == TEXT("sequence"))
+			{
+				CreatedNode->ReconstructNode();
+			}
 			
 			TempIdToNode.Add(TempId, CreatedNode);
 			
diff --git a/Engine/Plugins/Experimental/ECABridge/Source/ECABridge/Private/Commands/ECAUMGCommands.cpp b/Engine/Plugins/Experimental/ECABridge/Source/ECABridge/Private/Commands/ECAUMGCommands.cpp
--- a/Engine/Plugins/Experimental/ECABridge/Source/ECABridge/Private/Commands/ECAUMGCommands.cpp
+++ b/Engine/Plugins/Experimental/ECABridge/Source/ECABridge/Private/Commands/ECAUMGCommands.cpp
@@ -588,34 +588,122 @@ FECACommandResult FECACommand_GetWidgetInfo::Execute(const TSharedPtr<FJsonObjec
 	{
 		return FECACommandResult::Error(TEXT("Missing required parameter: widget_path"));
 	}
-	
+
 	UWidgetBlueprint* WidgetBlueprint = LoadUMGWidgetBlueprintByPath(WidgetPath);
 	if (!WidgetBlueprint)
 	{
 		return FECACommandResult::Error(FString::Printf(TEXT("Widget Blueprint not found: %s"), *WidgetPath));
 	}
-	
+
+	// Helper: FMargin → JSON {left,top,right,bottom}
+	auto MarginJson = [](const FMargin& M) -> TSharedPtr<FJsonObject>
+	{
+		TSharedPtr<FJsonObject> J = MakeShared<FJsonObject>();
+		J->SetNumberField(TEXT("left"),   M.Left);
+		J->SetNumberField(TEXT("top"),    M.Top);
+		J->SetNumberField(TEXT("right"),  M.Right);
+		J->SetNumberField(TEXT("bottom"), M.Bottom);
+		return J;
+	};
+
+	// Build a JSON object for a single widget including inline slot data.
+	// Returns: name, class, is_variable, parent, slot_type, and all slot layout properties.
+	auto WidgetToJson = [&MarginJson](UWidget* Widget) -> TSharedPtr<FJsonObject>
+	{
+		TSharedPtr<FJsonObject> W = MakeShared<FJsonObject>();
+		W->SetStringField(TEXT("name"),        Widget->GetName());
+		// For Blueprint-generated classes, return the full object path (e.g. /Game/UI/WBP_Foo.WBP_Foo_C)
+		// so create_widget_from_spec can resolve them via LoadClass. For native classes, GetName() is sufficient.
+		UClass* WidgetClass = Widget->GetClass();
+		FString ClassName = WidgetClass->IsNative()
+			? WidgetClass->GetName()
+			: WidgetClass->GetPathName();
+		W->SetStringField(TEXT("class"), ClassName);
+		W->SetBoolField  (TEXT("is_variable"), Widget->bIsVariable);
+
+		UPanelWidget* Parent = Widget->GetParent();
+		W->SetStringField(TEXT("parent"), Parent ? Parent->GetName() : TEXT(""));
+
+		// Child index within parent — needed to reconstruct correct order
+		if (Parent)
+		{
+			int32 ChildIdx = INDEX_NONE;
+			for (int32 i = 0; i < Parent->GetChildrenCount(); ++i)
+			{
+				if (Parent->GetChildAt(i) == Widget) { ChildIdx = i; break; }
+			}
+			W->SetNumberField(TEXT("child_index"), ChildIdx);
+		}
+
+		// Brush data for Image widgets
+		if (UImage* Img = Cast<UImage>(Widget))
+		{
+			const FSlateBrush& Brush = Img->GetBrush();
+			TSharedPtr<FJsonObject> BrushJ = MakeShared<FJsonObject>();
+			if (UObject* Resource = Brush.GetResourceObject())
+				BrushJ->SetStringField(TEXT("asset_path"), Resource->GetPathName());
+			const FLinearColor Tint = Brush.TintColor.GetSpecifiedColor();
+			TSharedPtr<FJsonObject> TintJ = MakeShared<FJsonObject>();
+			TintJ->SetNumberField(TEXT("r"), Tint.R);
+			TintJ->SetNumberField(TEXT("g"), Tint.G);
+			TintJ->SetNumberField(TEXT("b"), Tint.B);
+			TintJ->SetNumberField(TEXT("a"), Tint.A);
+			BrushJ->SetObjectField(TEXT("tint"), TintJ);
+			W->SetObjectField(TEXT("brush"), BrushJ);
+		}
+		if (UPanelSlot* Slot = Widget->Slot)
+		{
+			W->SetStringField(TEXT("slot_type"), Slot->GetClass()->GetName());
+
+			// Use reflection to read Padding, h_align, v_align where available
+			// (Avoids depending on private slot headers not available in FN_Main)
+			if (FStructProperty* PaddingProp = CastField<FStructProperty>(Slot->GetClass()->FindPropertyByName(TEXT("Padding"))))
+			{
+				if (FMargin* PaddingPtr = PaddingProp->ContainerPtrToValuePtr<FMargin>(Slot))
+				{
+					W->SetObjectField(TEXT("padding"), MarginJson(*PaddingPtr));
+				}
+			}
+			if (FByteProperty* HProp = CastField<FByteProperty>(Slot->GetClass()->FindPropertyByName(TEXT("HorizontalAlignment"))))
+				W->SetStringField(TEXT("h_align"), StaticEnum<EHorizontalAlignment>()->GetNameStringByValue((int64)HProp->GetPropertyValue_InContainer(Slot)));
+			if (FByteProperty* VProp = CastField<FByteProperty>(Slot->GetClass()->FindPropertyByName(TEXT("VerticalAlignment"))))
+				W->SetStringField(TEXT("v_align"), StaticEnum<EVerticalAlignment>()->GetNameStringByValue((int64)VProp->GetPropertyValue_InContainer(Slot)));
+
+			// CanvasPanel: layout, auto_size, z_order via public API
+			if (UCanvasPanelSlot* CpS = Cast<UCanvasPanelSlot>(Slot))
+			{
+				const FAnchorData& Layout = CpS->GetLayout();
+				TSharedPtr<FJsonObject> AnchorsJ = MakeShared<FJsonObject>();
+				AnchorsJ->SetNumberField(TEXT("min_x"), Layout.Anchors.Minimum.X);
+				AnchorsJ->SetNumberField(TEXT("min_y"), Layout.Anchors.Minimum.Y);
+				AnchorsJ->SetNumberField(TEXT("max_x"), Layout.Anchors.Maximum.X);
+				AnchorsJ->SetNumberField(TEXT("max_y"), Layout.Anchors.Maximum.Y);
+				W->SetObjectField(TEXT("anchors"),  AnchorsJ);
+				W->SetObjectField(TEXT("offsets"),  MarginJson(Layout.Offsets));
+				W->SetBoolField  (TEXT("auto_size"), CpS->GetAutoSize());
+				W->SetNumberField(TEXT("z_order"),   CpS->GetZOrder());
+			}
+		}
+
+		return W;
+	};
+
 	TSharedPtr<FJsonObject> Result = MakeResult();
-	Result->SetStringField(TEXT("name"), WidgetBlueprint->GetName());
-	Result->SetStringField(TEXT("path"), WidgetBlueprint->GetPathName());
+	Result->SetStringField(TEXT("name"),         WidgetBlueprint->GetName());
+	Result->SetStringField(TEXT("path"),         WidgetBlueprint->GetPathName());
 	Result->SetStringField(TEXT("parent_class"), WidgetBlueprint->ParentClass ? WidgetBlueprint->ParentClass->GetName() : TEXT("None"));
-	
-	// Get widgets in tree
+
 	TArray<TSharedPtr<FJsonValue>> WidgetsArray;
 	UWidgetTree* WidgetTree = WidgetBlueprint->WidgetTree;
 	if (WidgetTree)
 	{
-		WidgetTree->ForEachWidget([&WidgetsArray](UWidget* Widget)
+		WidgetTree->ForEachWidget([&WidgetsArray, &WidgetToJson](UWidget* Widget)
 		{
-			TSharedPtr<FJsonObject> WidgetJson = MakeShared<FJsonObject>();
-			WidgetJson->SetStringField(TEXT("name"), Widget->GetName());
-			WidgetJson->SetStringField(TEXT("class"), Widget->GetClass()->GetName());
-			WidgetJson->SetBoolField(TEXT("is_variable"), Widget->bIsVariable);
-			WidgetsArray.Add(MakeShared<FJsonValueObject>(WidgetJson));
+			WidgetsArray.Add(MakeShared<FJsonValueObject>(WidgetToJson(Widget)));
 		});
 	}
-	Result->SetArrayField(TEXT("widgets"), WidgetsArray);
+	Result->SetArrayField(TEXT("widgets"),      WidgetsArray);
 	Result->SetNumberField(TEXT("widget_count"), WidgetsArray.Num());
-	
+
 	return FECACommandResult::Success(Result);
 }
diff --git a/Engine/Plugins/Experimental/ECABridge/Source/ECABridge/Public/Commands/ECAUMGCommands.h b/Engine/Plugins/Experimental/ECABridge/Source/ECABridge/Public/Commands/ECAUMGCommands.h
--- a/Engine/Plugins/Experimental/ECABridge/Source/ECABridge/Public/Commands/ECAUMGCommands.h
+++ b/Engine/Plugins/Experimental/ECABridge/Source/ECABridge/Public/Commands/ECAUMGCommands.h
@@ -175,7 +175,7 @@ class FECACommand_GetWidgetInfo : public IECACommand
 {
 public:
 	virtual FString GetName() const override { return TEXT("get_widget_info"); }
-	virtual FString GetDescription() const override { return TEXT("Get information about a Widget Blueprint"); }
+	virtual FString GetDescription() const override { return TEXT("Get information about a Widget Blueprint including the full widget list with parent, class, is_variable, slot_type, and inline slot layout properties (padding, alignment, size rule, anchors)."); }
 	virtual FString GetCategory() const override { return TEXT("UMG"); }
 	
 	virtual TArray<FECACommandParam> GetParameters() const override
