[ECABridge] Widget construction from JSON spec

Enables full UMG widget hierarchies to be created from a single JSON
specification, replacing dozens of sequential MCP calls with one batched
operation. This cuts widget reconstruction time dramatically and makes
AI-driven widget authoring practical at scale.

New tool: create_widget_from_spec. Supports all slot types (Canvas,
Overlay, VBox, HBox, SizeBox, StackBox), nested children, image brush
properties, and post-buffer config.

#rnx
#Synced-CL 52344065

diff --git a/Engine/Plugins/Experimental/ECABridge/Source/ECABridge/Private/Commands/ECAWidgetTreeCommands.cpp b/Engine/Plugins/Experimental/ECABridge/Source/ECABridge/Private/Commands/ECAWidgetTreeCommands.cpp
--- a/Engine/Plugins/Experimental/ECABridge/Source/ECABridge/Private/Commands/ECAWidgetTreeCommands.cpp
+++ b/Engine/Plugins/Experimental/ECABridge/Source/ECABridge/Private/Commands/ECAWidgetTreeCommands.cpp
@@ -39,7 +39,15 @@
 #include "Components/CanvasPanelSlot.h"
 
 // Engine - Blueprint utilities
+#include "Editor.h"
+#include "FileHelpers.h"
 #include "Kismet2/BlueprintEditorUtils.h"
+#include "Kismet2/KismetEditorUtilities.h"
+#include "MVVMBlueprintView.h"
+#include "MVVMBlueprintViewBinding.h"
+#include "MVVMBlueprintViewModelContext.h"
+#include "MVVMEditorSubsystem.h"
+#include "MVVMPropertyPath.h"
 #include "UObject/UObjectIterator.h"
 #include "UObject/SavePackage.h"
 
@@ -67,6 +75,8 @@ REGISTER_ECA_COMMAND(FECACommand_RenameWidget)
 REGISTER_ECA_COMMAND(FECACommand_GetSlotProperty)
 REGISTER_ECA_COMMAND(FECACommand_CreateWidgetAnimation)
 
+REGISTER_ECA_COMMAND(FECACommand_SetWidgetDesignSize)
+
 //------------------------------------------------------------------------------
 // Shared Helper Functions
 //------------------------------------------------------------------------------
@@ -2169,3 +2179,596 @@ FECACommandResult FECACommand_CreateWidgetAnimation::Execute(const TSharedPtr<FJ
 
 	return FECACommandResult::Success(Result);
 }
+
+REGISTER_ECA_COMMAND(FECACommand_CreateWidgetFromSpec)
+
+//------------------------------------------------------------------------------
+// CreateWidgetFromSpec — helpers
+//------------------------------------------------------------------------------
+
+/** Recursively add widgets from a spec array. Errors are accumulated in OutErrors; created widget names in OutCreatedNames.
+ *  When ParentName is empty and no root widget exists yet, the first widget is assigned as WidgetBP->WidgetTree->RootWidget. */
+static void CreateWidgetsFromSpec(
+	UWidgetBlueprint* WidgetBP,
+	const TArray<TSharedPtr<FJsonValue>>* WidgetDefs,
+	const FString& ParentName,
+	TArray<FString>& OutCreatedNames,
+	TArray<FString>& OutErrors)
+{
+	if (!WidgetDefs) return;
+
+	for (const TSharedPtr<FJsonValue>& DefVal : *WidgetDefs)
+	{
+		if (!DefVal.IsValid() || DefVal->Type != EJson::Object) continue;
+		const TSharedPtr<FJsonObject>& Def = DefVal->AsObject();
+
+		FString WidgetName, WidgetClass;
+		Def->TryGetStringField(TEXT("name"), WidgetName);
+		Def->TryGetStringField(TEXT("class"), WidgetClass);
+		if (WidgetName.IsEmpty() || WidgetClass.IsEmpty())
+		{
+			OutErrors.Add(TEXT("Widget spec missing 'name' or 'class'"));
+			continue;
+		}
+
+		bool bIsVariable = false;
+		Def->TryGetBoolField(TEXT("is_variable"), bIsVariable);
+
+		// Find or create widget class
+		UClass* Class = FindObject<UClass>(nullptr, *FString::Printf(TEXT("/Script/UMG.%s"), *WidgetClass));
+		if (!Class) Class = LoadObject<UClass>(nullptr, *FString::Printf(TEXT("/Script/UMG.%s"), *WidgetClass));
+		if (!Class) Class = FindObject<UClass>(nullptr, *(WidgetClass + TEXT("_C")));
+		if (!Class) Class = LoadObject<UClass>(nullptr, *(WidgetClass + TEXT("_C")));
+		// WBP blueprint class fallback: try loading as a Blueprint-generated UClass via asset path.
+		// Handles bare names like 'WBP_Foo_C' and content paths like '/Game/UI/WBP_Foo_C'.
+		if (!Class)
+		{
+			FString ClassPath = WidgetClass;
+			if (!ClassPath.EndsWith(TEXT("_C"))) ClassPath += TEXT("_C");
+			Class = LoadClass<UUserWidget>(nullptr, *ClassPath);
+		}
+		if (!Class)
+		{
+			OutErrors.Add(FString::Printf(TEXT("Widget class not found: %s"), *WidgetClass));
+			continue;
+		}
+
+		// Create widget in tree
+		WidgetBP->Modify();
+		WidgetBP->WidgetTree->Modify();
+		UWidget* NewWidget = WidgetBP->WidgetTree->ConstructWidget<UWidget>(Class, FName(*WidgetName));
+		if (!NewWidget)
+		{
+			OutErrors.Add(FString::Printf(TEXT("Failed to create widget: %s (%s)"), *WidgetName, *WidgetClass));
+			continue;
+		}
+		NewWidget->bIsVariable = bIsVariable;
+		NewWidget->SetFlags(RF_Transactional);
+
+		// Add to parent (or set as root if ParentName is empty)
+		if (!ParentName.IsEmpty())
+		{
+			UWidget* ParentWidget = WidgetBP->WidgetTree->FindWidget(FName(*ParentName));
+			UPanelWidget* Panel = Cast<UPanelWidget>(ParentWidget);
+			if (!Panel)
+			{
+				OutErrors.Add(FString::Printf(
+					TEXT("Widget '%s': parent '%s' not found or is not a panel widget"),
+					*WidgetName, *ParentName));
+				WidgetBP->WidgetTree->RemoveWidget(NewWidget);
+				continue;
+			}
+
+			// Use InsertChildAt to preserve child order from the spec
+			double ChildIndexD = -1.0;
+			Def->TryGetNumberField(TEXT("child_index"), ChildIndexD);
+			int32 ChildIndex = (ChildIndexD >= 0) ? static_cast<int32>(ChildIndexD) : Panel->GetChildrenCount();
+			UPanelSlot* Slot = Panel->InsertChildAt(ChildIndex, NewWidget);
+			if (!Slot) Slot = Panel->AddChild(NewWidget); // fallback
+			if (!Slot)
+			{
+				OutErrors.Add(FString::Printf(
+					TEXT("Widget '%s': failed to add to panel"), *WidgetName));
+				WidgetBP->WidgetTree->RemoveWidget(NewWidget);
+				continue;
+			}
+
+			// Apply slot properties
+			const TSharedPtr<FJsonObject>* SlotDef = nullptr;
+			if (Slot && Def->TryGetObjectField(TEXT("slot"), SlotDef) && SlotDef && SlotDef->IsValid())
+			{
+				FString HAlign, VAlign, SizeRule;
+				double SizeValue = 1.0, PadL = 0, PadT = 0, PadR = 0, PadB = 0;
+				(*SlotDef)->TryGetStringField(TEXT("h_align"), HAlign);
+				(*SlotDef)->TryGetStringField(TEXT("v_align"), VAlign);
+				(*SlotDef)->TryGetStringField(TEXT("size_rule"), SizeRule);
+				(*SlotDef)->TryGetNumberField(TEXT("size_value"), SizeValue);
+				const TSharedPtr<FJsonObject>* PadObj = nullptr;
+				if ((*SlotDef)->TryGetObjectField(TEXT("padding"), PadObj) && PadObj)
+				{
+					(*PadObj)->TryGetNumberField(TEXT("left"),   PadL);
+					(*PadObj)->TryGetNumberField(TEXT("top"),    PadT);
+					(*PadObj)->TryGetNumberField(TEXT("right"),  PadR);
+					(*PadObj)->TryGetNumberField(TEXT("bottom"), PadB);
+				}
+				auto ApplyHA = [](const FString& S) -> EHorizontalAlignment {
+					if (S == TEXT("Fill"))  return HAlign_Fill;
+					if (S == TEXT("Left"))  return HAlign_Left;
+					if (S == TEXT("Center")) return HAlign_Center;
+					if (S == TEXT("Right")) return HAlign_Right;
+					return HAlign_Fill;
+				};
+				auto ApplyVA = [](const FString& S) -> EVerticalAlignment {
+					if (S == TEXT("Fill"))   return VAlign_Fill;
+					if (S == TEXT("Top"))    return VAlign_Top;
+					if (S == TEXT("Center")) return VAlign_Center;
+					if (S == TEXT("Bottom")) return VAlign_Bottom;
+					return VAlign_Fill;
+				};
+				FMargin Pad(PadL, PadT, PadR, PadB);
+				bool bFill = SizeRule.Equals(TEXT("Fill"), ESearchCase::IgnoreCase);
+				FSlateChildSize SlateSize(bFill ? ESlateSizeRule::Fill : ESlateSizeRule::Automatic);
+				SlateSize.Value = static_cast<float>(SizeValue);
+
+				if      (UOverlaySlot*      OvS  = Cast<UOverlaySlot>(Slot))      { OvS->SetPadding(Pad);  OvS->SetHorizontalAlignment(ApplyHA(HAlign)); OvS->SetVerticalAlignment(ApplyVA(VAlign)); }
+				else if (UVerticalBoxSlot*   VbS  = Cast<UVerticalBoxSlot>(Slot))   { VbS->SetPadding(Pad);  VbS->SetHorizontalAlignment(ApplyHA(HAlign)); VbS->SetVerticalAlignment(ApplyVA(VAlign)); VbS->SetSize(SlateSize); }
+				else if (UHorizontalBoxSlot* HbS  = Cast<UHorizontalBoxSlot>(Slot)) { HbS->SetPadding(Pad);  HbS->SetHorizontalAlignment(ApplyHA(HAlign)); HbS->SetVerticalAlignment(ApplyVA(VAlign)); HbS->SetSize(SlateSize); }
+				else if (USizeBoxSlot*       SbS  = Cast<USizeBoxSlot>(Slot))       { SbS->SetPadding(Pad);  SbS->SetHorizontalAlignment(ApplyHA(HAlign)); SbS->SetVerticalAlignment(ApplyVA(VAlign)); }
+				else if (UCanvasPanelSlot*   CpS  = Cast<UCanvasPanelSlot>(Slot))
+				{
+					const TSharedPtr<FJsonObject>* AnchorObj = nullptr;
+					const TSharedPtr<FJsonObject>* OffsetObj  = nullptr;
+					if ((*SlotDef)->TryGetObjectField(TEXT("anchors"), AnchorObj) && AnchorObj)
+					{
+						FAnchorData Layout;
+						double MinX=0,MinY=0,MaxX=0,MaxY=0;
+						(*AnchorObj)->TryGetNumberField(TEXT("min_x"), MinX);
+						(*AnchorObj)->TryGetNumberField(TEXT("min_y"), MinY);
+						(*AnchorObj)->TryGetNumberField(TEXT("max_x"), MaxX);
+						(*AnchorObj)->TryGetNumberField(TEXT("max_y"), MaxY);
+						Layout.Anchors = FAnchors(MinX, MinY, MaxX, MaxY);
+						if ((*SlotDef)->TryGetObjectField(TEXT("offsets"), OffsetObj) && OffsetObj)
+						{
+							double OL=0,OT=0,OR=0,OB=0;
+							(*OffsetObj)->TryGetNumberField(TEXT("left"),   OL);
+							(*OffsetObj)->TryGetNumberField(TEXT("top"),    OT);
+							(*OffsetObj)->TryGetNumberField(TEXT("right"),  OR);
+							(*OffsetObj)->TryGetNumberField(TEXT("bottom"), OB);
+							Layout.Offsets = FMargin(OL,OT,OR,OB);
+						}
+						CpS->SetLayout(Layout);
+					}
+				}
+				Slot->SynchronizeProperties();
+			}
+		}
+		else
+		{
+			// Root widget — only one root is supported; error if one already exists.
+			// Destroy the just-created widget to avoid orphaning it in the tree.
+			if (WidgetBP->WidgetTree->RootWidget != nullptr)
+			{
+				OutErrors.Add(FString::Printf(
+					TEXT("Widget '%s' cannot be root: a root widget already exists ('%s'). Only the first top-level widget in the spec can be the root."),
+					*WidgetName, *WidgetBP->WidgetTree->RootWidget->GetName()));
+				WidgetBP->WidgetTree->RemoveWidget(NewWidget);
+				continue;
+			}
+			WidgetBP->WidgetTree->RootWidget = NewWidget;
+		}
+		OutCreatedNames.Add(WidgetName);
+
+		// Apply brush properties to Image widgets
+		const TSharedPtr<FJsonObject>* BrushDef = nullptr;
+		if (UImage* ImageWidget = Cast<UImage>(NewWidget))
+		{
+			if (Def->TryGetObjectField(TEXT("brush"), BrushDef) && BrushDef && BrushDef->IsValid())
+			{
+				// Apply asset (must come before GetBrush() snapshot so resource is captured)
+				FString AssetPath;
+				if ((*BrushDef)->TryGetStringField(TEXT("asset_path"), AssetPath))
+				{
+					if (UObject* Asset = LoadObject<UObject>(nullptr, *AssetPath))
+					{
+						if (UMaterialInterface* Mat = Cast<UMaterialInterface>(Asset)) ImageWidget->SetBrushFromMaterial(Mat);
+						else if (UTexture2D* Tex = Cast<UTexture2D>(Asset)) ImageWidget->SetBrushFromTexture(Tex);
+						else ImageWidget->SetBrushResourceObject(Asset);
+					}
+				}
+
+				// Snapshot AFTER asset so resource is preserved when we call SetBrush()
+				FSlateBrush ModifiedBrush = ImageWidget->GetBrush();
+				bool bBrushModified = false;
+
+				const TSharedPtr<FJsonObject>* TintObj = nullptr;
+				if ((*BrushDef)->TryGetObjectField(TEXT("tint"), TintObj) && TintObj && TintObj->IsValid())
+				{
+					double R=1,G=1,B=1,A=1;
+					(*TintObj)->TryGetNumberField(TEXT("r"),R); (*TintObj)->TryGetNumberField(TEXT("g"),G);
+					(*TintObj)->TryGetNumberField(TEXT("b"),B); (*TintObj)->TryGetNumberField(TEXT("a"),A);
+					ModifiedBrush.TintColor = FSlateColor(FLinearColor(R,G,B,A));
+					bBrushModified = true;
+				}
+
+				FString DrawAsStr;
+				TEnumAsByte<ESlateBrushDrawType::Type> DrawAsVal;
+				if ((*BrushDef)->TryGetStringField(TEXT("draw_as"), DrawAsStr) && ParseDrawType(DrawAsStr, DrawAsVal))
+				{
+					ModifiedBrush.DrawAs = DrawAsVal;
+					bBrushModified = true;
+				}
+
+				if (bBrushModified)
+				{
+					ImageWidget->SetBrush(ModifiedBrush);
+				}
+			}
+		}
+
+		// Apply post_buffer config to PostBufferUpdate widgets
+		const TSharedPtr<FJsonObject>* PBDef = nullptr;
+		if (UPostBufferUpdate* PBWidget = Cast<UPostBufferUpdate>(NewWidget))
+		{
+			if (Def->TryGetObjectField(TEXT("post_buffer"), PBDef) && PBDef && PBDef->IsValid())
+			{
+				// Delegate to the existing reflection-based approach used in set_post_buffer_update_config
+				FProperty* Prop = UPostBufferUpdate::StaticClass()->FindPropertyByName(TEXT("UpdateBufferInfos"));
+				FArrayProperty* ArrayProp = Prop ? CastField<FArrayProperty>(Prop) : nullptr;
+				const TArray<TSharedPtr<FJsonValue>>* BuffersArr = nullptr;
+				if (ArrayProp && (*PBDef)->TryGetArrayField(TEXT("buffers"), BuffersArr) && BuffersArr)
+				{
+					double BlurStrength = 0.0;
+					bool bHasBlur = (*PBDef)->TryGetNumberField(TEXT("blur_strength"), BlurStrength) && BlurStrength > 0.0;
+					static const ESlatePostRT BufferEnumValues[] = {
+						ESlatePostRT::ESlatePostRT_0, ESlatePostRT::ESlatePostRT_1, ESlatePostRT::ESlatePostRT_2,
+						ESlatePostRT::ESlatePostRT_3, ESlatePostRT::ESlatePostRT_4
+					};
+					TArray<FSlatePostBufferUpdateInfo> Infos;
+					for (const TSharedPtr<FJsonValue>& BV : *BuffersArr)
+					{
+						if (!BV.IsValid()) continue;
+						double Idx = 0; BV->TryGetNumber(Idx);
+						int32 I = FMath::Clamp((int32)Idx, 0, 4);
+						FSlatePostBufferUpdateInfo Info;
+						Info.BufferToUpdate = BufferEnumValues[I];
+						if (bHasBlur)
+						{
+							UPostBufferBlurUpdater* Blur = NewObject<UPostBufferBlurUpdater>(WidgetBP);
+							Blur->GaussianBlurStrength = static_cast<float>(BlurStrength);
+							Info.PostParamUpdater = Blur;
+						}
+						Infos.Add(Info);
+					}
+					PBWidget->Modify();
+					TArray<FSlatePostBufferUpdateInfo>* TargetArray = ArrayProp->ContainerPtrToValuePtr<TArray<FSlatePostBufferUpdateInfo>>(PBWidget);
+					*TargetArray = MoveTemp(Infos);
+					FPropertyChangedEvent PChange(Prop, EPropertyChangeType::ValueSet);
+					PBWidget->PostEditChangeProperty(PChange);
+				}
+			}
+		}
+
+		// Recurse into children
+		const TArray<TSharedPtr<FJsonValue>>* ChildDefs = nullptr;
+		if (Def->TryGetArrayField(TEXT("children"), ChildDefs) && ChildDefs)
+		{
+			CreateWidgetsFromSpec(WidgetBP, ChildDefs, WidgetName, OutCreatedNames, OutErrors);
+		}
+	}
+}
+
+//------------------------------------------------------------------------------
+// CreateWidgetFromSpec
+//------------------------------------------------------------------------------
+
+FECACommandResult FECACommand_CreateWidgetFromSpec::Execute(const TSharedPtr<FJsonObject>& Params)
+{
+	FString WidgetPath;
+	if (!GetStringParam(Params, TEXT("widget_path"), WidgetPath))
+	{
+		return FECACommandResult::Error(TEXT("Missing required parameter: widget_path"));
+	}
+
+	FString ParentClassName = TEXT("UserWidget");
+	GetStringParam(Params, TEXT("parent_class"), ParentClassName, false);
+
+	// Resolve parent class
+	UClass* ParentClass = FindObject<UClass>(nullptr, *FString::Printf(TEXT("/Script/UMG.%s"), *ParentClassName));
+	if (!ParentClass) ParentClass = FindObject<UClass>(nullptr, *ParentClassName);
+	if (!ParentClass) ParentClass = LoadClass<UUserWidget>(nullptr, *FString::Printf(TEXT("/Script/UMG.%s"), *ParentClassName));
+	if (!ParentClass) ParentClass = UUserWidget::StaticClass();
+
+	// Split path into package path and asset name. Both components are required.
+	int32 SlashIdx;
+	if (!WidgetPath.FindLastChar('/', SlashIdx) || SlashIdx == WidgetPath.Len() - 1)
+	{
+		return FECACommandResult::Error(TEXT("widget_path must be a full asset path with at least one '/' (e.g. /Game/UI/WBP_MyWidget)"));
+	}
+	const FString AssetName = WidgetPath.Mid(SlashIdx + 1);
+
+	// In UE, CreatePackage takes the FULL asset path (package name == asset path, e.g. /Game/UI/WBP_Foo).
+	UPackage* Package = CreatePackage(*WidgetPath);
+	if (!Package)
+	{
+		return FECACommandResult::Error(FString::Printf(TEXT("Failed to create package: %s"), *WidgetPath));
+	}
+
+	UWidgetBlueprint* WidgetBP = Cast<UWidgetBlueprint>(FKismetEditorUtilities::CreateBlueprint(
+		ParentClass, Package, FName(*AssetName),
+		BPTYPE_Normal, UWidgetBlueprint::StaticClass(),
+		UWidgetBlueprintGeneratedClass::StaticClass()
+	));
+	if (!WidgetBP)
+	{
+		return FECACommandResult::Error(FString::Printf(TEXT("Failed to create Widget Blueprint at: %s"), *WidgetPath));
+	}
+
+	TArray<FString> CreatedWidgets;
+	TArray<FString> Errors;
+
+	// Build widget hierarchy
+	const TArray<TSharedPtr<FJsonValue>>* WidgetDefs = nullptr;
+	if (Params->TryGetArrayField(TEXT("widgets"), WidgetDefs) && WidgetDefs)
+	{
+		CreateWidgetsFromSpec(WidgetBP, WidgetDefs, TEXT(""), CreatedWidgets, Errors);
+	}
+
+	// Add ViewModels
+	const TArray<TSharedPtr<FJsonValue>>* VMDefs = nullptr;
+	if (Params->TryGetArrayField(TEXT("viewmodels"), VMDefs) && VMDefs)
+	{
+		UMVVMEditorSubsystem* MVVMSubsystem = GEditor ? GEditor->GetEditorSubsystem<UMVVMEditorSubsystem>() : nullptr;
+		if (MVVMSubsystem)
+		{
+			for (const TSharedPtr<FJsonValue>& VMVal : *VMDefs)
+			{
+				if (!VMVal.IsValid() || VMVal->Type != EJson::Object) continue;
+				const TSharedPtr<FJsonObject>& VMDef = VMVal->AsObject();
+				FString VMName, VMClass;
+				VMDef->TryGetStringField(TEXT("name"),  VMName);
+				VMDef->TryGetStringField(TEXT("class"), VMClass);
+				if (VMClass.IsEmpty()) continue; // name is optional — defaults to class short name
+
+				UClass* Class = FindObject<UClass>(nullptr, *VMClass);
+				if (!Class) Class = LoadObject<UClass>(nullptr, *VMClass);
+				if (!Class) { Errors.Add(FString::Printf(TEXT("ViewModel class not found: %s"), *VMClass)); continue; }
+
+				// FN_Main API registers the context under the class short name.
+				// VMName is informational only in this API version — bindings should use the class short name.
+				MVVMSubsystem->AddViewModel(WidgetBP, Class);
+			}
+		}
+	}
+
+	// Initial compile to ensure GeneratedClass exists before we access the CDO.
+	// design_size_mode must be set on the CDO before the main compile so it serializes correctly.
+	EnsureAllWidgetGuids(WidgetBP);
+	FKismetEditorUtilities::CompileBlueprint(WidgetBP);
+
+	// Set DesignSizeMode on the CDO — GeneratedClass is now guaranteed non-null after the compile above.
+	{
+		FString DesignModeName;
+		if (GetStringParam(Params, TEXT("design_size_mode"), DesignModeName, false))
+		{
+			UUserWidget* CDO = WidgetBP->GeneratedClass ? Cast<UUserWidget>(WidgetBP->GeneratedClass->GetDefaultObject(true)) : nullptr;
+			if (!CDO)
+			{
+				Errors.Add(TEXT("design_size_mode: CDO unavailable after initial compile — skipped"));
+			}
+			else
+			{
+				CDO->Modify();
+				if (DesignModeName.Equals(TEXT("FillScreen"), ESearchCase::IgnoreCase))
+					CDO->DesignSizeMode = EDesignPreviewSizeMode::FillScreen;
+				else if (DesignModeName.Equals(TEXT("Desired"), ESearchCase::IgnoreCase))
+					CDO->DesignSizeMode = EDesignPreviewSizeMode::Desired;
+				else
+					CDO->DesignSizeMode = EDesignPreviewSizeMode::Custom;
+			}
+		}
+	}
+
+	// Re-compile after CDO changes so MVVM GeneratedClass is fully up to date for self-bindings.
+	FKismetEditorUtilities::CompileBlueprint(WidgetBP);
+
+
+	// Add MVVM bindings
+	const TArray<TSharedPtr<FJsonValue>>* BindingDefs = nullptr;
+	if (Params->TryGetArrayField(TEXT("bindings"), BindingDefs) && BindingDefs)
+	{
+		UMVVMEditorSubsystem* MVVMSubsystem = GEditor ? GEditor->GetEditorSubsystem<UMVVMEditorSubsystem>() : nullptr;
+		UMVVMBlueprintView* View = MVVMSubsystem ? MVVMSubsystem->GetView(WidgetBP) : nullptr;
+		if (MVVMSubsystem && View)
+		{
+			for (const TSharedPtr<FJsonValue>& BindVal : *BindingDefs)
+			{
+				if (!BindVal.IsValid() || BindVal->Type != EJson::Object) continue;
+				const TSharedPtr<FJsonObject>& BD = BindVal->AsObject();
+				FString WidgetName, DestProp, VMName, VMProp;
+				BD->TryGetStringField(TEXT("widget_name"), WidgetName);
+				BD->TryGetStringField(TEXT("property"),    DestProp);
+				BD->TryGetStringField(TEXT("vm_name"),     VMName);
+				BD->TryGetStringField(TEXT("vm_property"), VMProp);
+				if (DestProp.IsEmpty() || VMProp.IsEmpty()) continue;
+
+				// Resolve ViewModel context
+				const FMVVMBlueprintViewModelContext* VMContext = nullptr;
+				if (!VMName.IsEmpty())
+				{
+					VMContext = View->FindViewModel(FName(*VMName));
+					// FN_Main API registers VMs by class short name. Fall back to class-name lookup
+					// so specs where VMName matches the class short name (the common case) always work.
+					if (!VMContext)
+					{
+						for (const FMVVMBlueprintViewModelContext& Ctx : View->GetViewModels())
+						{
+							if (const UClass* C = Ctx.GetViewModelClass())
+							{
+								if (C->GetName().Equals(VMName, ESearchCase::IgnoreCase))
+								{ VMContext = &Ctx; break; }
+							}
+						}
+					}
+					if (!VMContext)
+					{
+						Errors.Add(FString::Printf(TEXT("Binding: ViewModel '%s' not found — skipping"), *VMName));
+						continue;
+					}
+				}
+				else if (View->GetViewModels().Num() > 0)
+				{
+					// vm_name omitted — use first registered ViewModel
+					VMContext = &View->GetViewModels()[0];
+				}
+				if (!VMContext) { Errors.Add(TEXT("Binding: no ViewModels registered on widget")); continue; }
+
+				const UClass* VMClass = VMContext->GetViewModelClass();
+				if (!VMClass) continue;
+
+				FProperty* VMPropPtr = VMClass ? FindFProperty<FProperty>(VMClass, *VMProp) : nullptr;
+				if (!VMPropPtr) { Errors.Add(FString::Printf(TEXT("Binding: property '%s' not found on ViewModel '%s' — skipping"), *VMProp, *VMName)); continue; }
+
+				FMVVMBlueprintViewBinding& NewBinding = MVVMSubsystem->AddBinding(WidgetBP);
+				MVVMSubsystem->SetBindingTypeForBinding(WidgetBP, NewBinding, EMVVMBindingMode::OneWayToDestination);
+
+				FMVVMBlueprintPropertyPath SrcPath;
+				SrcPath.SetViewModelId(VMContext->GetViewModelId());
+				SrcPath.SetPropertyPath(WidgetBP, UE::MVVM::FMVVMConstFieldVariant(VMPropPtr));
+				MVVMSubsystem->SetSourcePathForBinding(WidgetBP, NewBinding, SrcPath);
+
+				// Destination path — skip binding with error if destination cannot be resolved
+				bool bSelf = WidgetName.IsEmpty() || WidgetName.Equals(TEXT("self"), ESearchCase::IgnoreCase);
+				FMVVMBlueprintPropertyPath DestPath;
+				bool bDestResolved = false;
+				if (bSelf)
+				{
+					const UClass* DestClass = WidgetBP->GeneratedClass;
+					if (FProperty* DestPropPtr = DestClass ? FindFProperty<FProperty>(DestClass, *DestProp) : nullptr)
+					{
+						DestPath.SetSelfContext();
+						DestPath.SetPropertyPath(WidgetBP, UE::MVVM::FMVVMConstFieldVariant(DestPropPtr));
+						bDestResolved = true;
+					}
+				}
+				else
+				{
+					UWidget* DestWidget = WidgetBP->WidgetTree->FindWidget(FName(*WidgetName));
+					if (DestWidget)
+					{
+						if (FProperty* DestPropPtr = FindFProperty<FProperty>(DestWidget->GetClass(), *DestProp))
+						{
+							DestPath.SetWidgetName(FName(*WidgetName));
+							DestPath.SetPropertyPath(WidgetBP, UE::MVVM::FMVVMConstFieldVariant(DestPropPtr));
+							bDestResolved = true;
+						}
+					}
+				}
+				if (!bDestResolved)
+				{
+					Errors.Add(FString::Printf(
+						TEXT("Binding: destination '%s.%s' not found — skipping"), *WidgetName, *DestProp));
+					MVVMSubsystem->RemoveBinding(WidgetBP, NewBinding);
+					continue;
+				}
+				MVVMSubsystem->SetDestinationPathForBinding(WidgetBP, NewBinding, DestPath, false);
+			}
+		}
+	}
+
+	// Re-compile after MVVM bindings are set up, then save
+	EnsureAllWidgetGuids(WidgetBP);
+	FKismetEditorUtilities::CompileBlueprint(WidgetBP);
+	UEditorLoadingAndSavingUtils::SavePackages({ WidgetBP->GetOutermost() }, false);
+
+	// Result
+	TSharedPtr<FJsonObject> Result = MakeResult();
+	Result->SetStringField(TEXT("widget_path"), WidgetBP->GetPathName());
+	Result->SetNumberField(TEXT("widgets_created"), CreatedWidgets.Num());
+	Result->SetBoolField  (TEXT("has_errors"), !Errors.IsEmpty());
+
+	if (!Errors.IsEmpty())
+	{
+		TArray<TSharedPtr<FJsonValue>> ErrArr;
+		for (const FString& E : Errors) ErrArr.Add(MakeShared<FJsonValueString>(E));
+		Result->SetArrayField(TEXT("errors"), ErrArr);
+	}
+
+	return FECACommandResult::Success(Result);
+}
+
+//------------------------------------------------------------------------------
+// SetWidgetDesignSize
+//------------------------------------------------------------------------------
+FECACommandResult FECACommand_SetWidgetDesignSize::Execute(const TSharedPtr<FJsonObject>& Params)
+{
+	FString WidgetPath;
+	if (!GetStringParam(Params, TEXT("widget_path"), WidgetPath))
+		return FECACommandResult::Error(TEXT("Missing required parameter: widget_path"));
+
+	UWidgetBlueprint* WidgetBP = LoadWidgetBlueprintFromPath(WidgetPath);
+	if (!WidgetBP)
+		return FECACommandResult::Error(FString::Printf(TEXT("Widget Blueprint not found: %s"), *WidgetPath));
+
+	// DesignTimeSize and DesignSizeMode live on the UUserWidget CDO.
+	// We must Modify() the CDO (marks it for serialization) before changing its properties.
+	UUserWidget* CDO = nullptr;
+	if (WidgetBP->GeneratedClass)
+		CDO = Cast<UUserWidget>(WidgetBP->GeneratedClass->GetDefaultObject(true));
+
+	if (!CDO)
+		return FECACommandResult::Error(TEXT("Failed to get CDO for widget blueprint"));
+
+	// Apply size if provided
+	double SizeX = -1, SizeY = -1;
+	bool bHasSizeX = Params->TryGetNumberField(TEXT("size_x"), SizeX);
+	bool bHasSizeY = Params->TryGetNumberField(TEXT("size_y"), SizeY);
+
+	// Apply mode if provided
+	FString ModeStr;
+	GetStringParam(Params, TEXT("mode"), ModeStr, false);
+	EDesignPreviewSizeMode NewMode = CDO->DesignSizeMode;
+	if (!ModeStr.IsEmpty())
+	{
+		if (ModeStr.Equals(TEXT("Custom"), ESearchCase::IgnoreCase))
+			NewMode = EDesignPreviewSizeMode::Custom;
+		else if (ModeStr.Equals(TEXT("FillScreen"), ESearchCase::IgnoreCase))
+			NewMode = EDesignPreviewSizeMode::FillScreen;
+		else if (ModeStr.Equals(TEXT("Desired"), ESearchCase::IgnoreCase))
+			NewMode = EDesignPreviewSizeMode::Desired;
+		else
+			return FECACommandResult::Error(FString::Printf(
+				TEXT("Invalid mode '%s'. Valid: Custom, FillScreen, Desired"), *ModeStr));
+	}
+
+	// Only modify/save if something actually changed
+	if (bHasSizeX || bHasSizeY || !ModeStr.IsEmpty())
+	{
+		// Modify() BEFORE mutations so the transaction system snapshots pre-change state.
+		CDO->Modify();
+		WidgetBP->Modify();
+
+		if (bHasSizeX) CDO->DesignTimeSize.X = static_cast<float>(SizeX);
+		if (bHasSizeY) CDO->DesignTimeSize.Y = static_cast<float>(SizeY);
+		if (!ModeStr.IsEmpty()) CDO->DesignSizeMode = NewMode;
+
+		// Compile so the CDO change is serialized into the package.
+		FKismetEditorUtilities::CompileBlueprint(WidgetBP);
+		UEditorLoadingAndSavingUtils::SavePackages({ WidgetBP->GetOutermost() }, false);
+	}
+
+	// Return current values (whether changed or just read)
+	FString CurrentMode;
+	switch (static_cast<uint8>(CDO->DesignSizeMode))
+	{
+		case static_cast<uint8>(EDesignPreviewSizeMode::Custom):     CurrentMode = TEXT("Custom");     break;
+		case static_cast<uint8>(EDesignPreviewSizeMode::FillScreen): CurrentMode = TEXT("FillScreen"); break;
+		case static_cast<uint8>(EDesignPreviewSizeMode::Desired):    CurrentMode = TEXT("Desired");    break;
+		default:                                                      CurrentMode = TEXT("Unknown");    break;
+	}
+
+	TSharedPtr<FJsonObject> Result = MakeResult();
+	Result->SetNumberField(TEXT("size_x"), CDO->DesignTimeSize.X);
+	Result->SetNumberField(TEXT("size_y"), CDO->DesignTimeSize.Y);
+	Result->SetStringField(TEXT("mode"), CurrentMode);
+	return FECACommandResult::Success(Result);
+}
diff --git a/Engine/Plugins/Experimental/ECABridge/Source/ECABridge/Public/Commands/ECAWidgetTreeCommands.h b/Engine/Plugins/Experimental/ECABridge/Source/ECABridge/Public/Commands/ECAWidgetTreeCommands.h
--- a/Engine/Plugins/Experimental/ECABridge/Source/ECABridge/Public/Commands/ECAWidgetTreeCommands.h
+++ b/Engine/Plugins/Experimental/ECABridge/Source/ECABridge/Public/Commands/ECAWidgetTreeCommands.h
@@ -313,3 +313,66 @@ public:
 
 	virtual FECACommandResult Execute(const TSharedPtr<FJsonObject>& Params) override;
 };
+
+/**
+ * Create a complete Widget Blueprint from a JSON spec in a single call.
+ * The spec describes the widget hierarchy, slot properties, optional MVVM viewmodels,
+ * and bindings. This replaces the N-call pattern (create_umg_widget_blueprint + N×add_child_widget
+ * + N×set_slot_property + set_widget_viewmodel + N×add_viewmodel_binding).
+ *
+ * Spec format:
+ *   widget_path     — destination path (e.g. "/Game/UI/WBP_MyWidget")
+ *   parent_class    — optional parent class name (default: UserWidget)
+ *   widgets         — array of widget descriptors (see below)
+ *   viewmodels      — optional array of {name, class} viewmodel contexts
+ *   bindings        — optional array of {widget_name, property, vm_name, vm_property}. All bindings use OneWayToDestination mode. Use add_viewmodel_binding for other modes or conversion functions.
+ *
+ * Widget descriptor fields:
+ *   name, class, is_variable, slot (h_align, v_align, padding, size_rule, size_value, anchors, offsets),
+ *   brush (asset_path, tint, draw_as), post_buffer (buffers, blur_strength), children[]
+ */
+class FECACommand_CreateWidgetFromSpec : public IECACommand
+{
+public:
+	virtual FString GetName() const override { return TEXT("create_widget_from_spec"); }
+	virtual FString GetDescription() const override { return TEXT("Create a complete Widget Blueprint (hierarchy + slots + MVVM) from a single JSON spec. Replaces the N-call create/add/set pattern with one call. Package/blueprint creation failures are fatal; all other errors (bad widget class, missing parent, etc.) are collected and returned while the blueprint is still saved."); }
+	virtual FString GetCategory() const override { return TEXT("UMG"); }
+
+	virtual TArray<FECACommandParam> GetParameters() const override
+	{
+		return {
+			{ TEXT("widget_path"),  TEXT("string"), TEXT("Destination asset path (e.g. /Game/UI/WBP_MyWidget)"), true },
+			{ TEXT("parent_class"), TEXT("string"), TEXT("Parent class name (default: UserWidget)"), false },
+			{ TEXT("widgets"),      TEXT("array"),  TEXT("Widget hierarchy descriptors. Each entry: {name, class, is_variable, slot:{h_align,v_align,padding,size_rule,size_value,anchors,offsets}, children:[...], brush:{asset_path,tint,draw_as}, post_buffer:{buffers,blur_strength}}"), false },
+			{ TEXT("viewmodels"),   TEXT("array"),  TEXT("ViewModel contexts: [{class, name(optional)}]. 'name' is informational only \u2014 the context is always registered under the class short name in FN_Main. Bindings must use the class short name in vm_name."), false },
+			{ TEXT("bindings"),          TEXT("array"),  TEXT("MVVM bindings: [{widget_name, property, vm_name, vm_property}]. All bindings use OneWayToDestination mode. For other modes or conversion functions, use add_viewmodel_binding after creation."), false },
+			{ TEXT("design_size_mode"), TEXT("string"), TEXT("Designer canvas size mode: FillScreen, Custom (UE default if omitted), or Desired"), false },
+		};
+	}
+
+	virtual FECACommandResult Execute(const TSharedPtr<FJsonObject>& Params) override;
+};
+
+/**
+ * Set the designer canvas size and mode on a Widget Blueprint.
+ * Affects how the widget appears in the WBP designer view.
+ */
+class FECACommand_SetWidgetDesignSize : public IECACommand
+{
+public:
+	virtual FString GetName()        const override { return TEXT("set_widget_design_size"); }
+	virtual FString GetDescription() const override { return TEXT("Set the designer canvas size (DesignTimeSize) and mode on a Widget Blueprint. Use mode='Custom' with size_x/size_y, 'FillScreen', or 'Desired'."); }
+	virtual FString GetCategory()    const override { return TEXT("UMG"); }
+
+	virtual TArray<FECACommandParam> GetParameters() const override
+	{
+		return {
+			{ TEXT("widget_path"), TEXT("string"), TEXT("Path to the Widget Blueprint"), true },
+			{ TEXT("size_x"),      TEXT("number"), TEXT("Design canvas width in pixels (e.g. 1920)"), false },
+			{ TEXT("size_y"),      TEXT("number"), TEXT("Design canvas height in pixels (e.g. 353)"), false },
+			{ TEXT("mode"),        TEXT("string"), TEXT("Design size mode: Custom, FillScreen, or Desired"), false },
+		};
+	}
+
+	virtual FECACommandResult Execute(const TSharedPtr<FJsonObject>& Params) override;
+};
