[ECABridge] save_all_dirty: add optional post-save asset validation

Adds a validate parameter to save_all_dirty. When true, the saved
packages are looked up in the Asset Registry and run through
UEditorValidatorSubsystem. The result always includes validation_errors
and validation_warnings fields when validate:true is passed, regardless
of whether the Asset Registry has indexed the packages yet or whether
the validator subsystem is available.

#rnx
#Synced-CL 52344065

diff --git a/Engine/Plugins/Experimental/ECABridge/Source/ECABridge/Private/Commands/ECASaveCommands.cpp b/Engine/Plugins/Experimental/ECABridge/Source/ECABridge/Private/Commands/ECASaveCommands.cpp
--- a/Engine/Plugins/Experimental/ECABridge/Source/ECABridge/Private/Commands/ECASaveCommands.cpp
+++ b/Engine/Plugins/Experimental/ECABridge/Source/ECABridge/Private/Commands/ECASaveCommands.cpp
@@ -1,6 +1,9 @@
 // Copyright Epic Games, Inc. All Rights Reserved.
 
 #include "Commands/ECASaveCommands.h"
+#include "AssetRegistry/AssetRegistryModule.h"
+#include "Editor.h"
+#include "EditorValidatorSubsystem.h"
 #include "Engine/Blueprint.h"
 #include "Kismet2/KismetEditorUtilities.h"
 #include "Misc/PackageName.h"
@@ -197,13 +200,61 @@ FECACommandResult FECACommand_SaveAllDirty::Execute(const TSharedPtr<FJsonObject
 	}
 
 	TSharedPtr<FJsonObject> Result = MakeResult();
-	Result->SetNumberField(TEXT("saved_count"), SavedCount);
-	Result->SetNumberField(TEXT("failed_count"), FailedCount);
-	Result->SetArrayField(TEXT("saved_packages"), SavedArray);
+
+	// Optional post-save validation
+	bool bValidate = false;
+	GetBoolParam(Params, TEXT("validate"), bValidate, false);
+	int32 ValidationErrors   = 0;
+	int32 ValidationWarnings = 0;
+	if (bValidate && SavedCount > 0)
+	{
+		UEditorValidatorSubsystem* ValidatorSubsystem = GEditor
+			? GEditor->GetEditorSubsystem<UEditorValidatorSubsystem>()
+			: nullptr;
+		if (ValidatorSubsystem)
+		{
+			FAssetRegistryModule* AssetRegistryModule = FModuleManager::GetModulePtr<FAssetRegistryModule>(TEXT("AssetRegistry"));
+			if (AssetRegistryModule)
+			{
+				IAssetRegistry& AssetRegistry = AssetRegistryModule->Get();
+				TArray<FAssetData> AssetsToValidate;
+				for (const TSharedPtr<FJsonValue>& SavedVal : SavedArray)
+				{
+					FString PkgName;
+					if (SavedVal.IsValid() && SavedVal->TryGetString(PkgName))
+					{
+						TArray<FAssetData> Found;
+						AssetRegistry.GetAssetsByPackageName(FName(*PkgName), Found);
+						AssetsToValidate.Append(Found);
+					}
+				}
+				if (!AssetsToValidate.IsEmpty())
+				{
+					FValidateAssetsSettings ValidationSettings;
+					ValidationSettings.bSkipExcludedDirectories = false;
+					ValidationSettings.bCollectPerAssetDetails  = false;
+					ValidationSettings.bShowIfNoFailures        = false;
+					FValidateAssetsResults ValidationResults;
+					ValidatorSubsystem->ValidateAssetsWithSettings(AssetsToValidate, ValidationSettings, ValidationResults);
+					ValidationErrors   = ValidationResults.NumInvalid;
+					ValidationWarnings = ValidationResults.NumWarnings;
+				}
+			}
+		}
+	}
+
+	Result->SetNumberField(TEXT("saved_count"),   SavedCount);
+	Result->SetNumberField(TEXT("failed_count"),  FailedCount);
+	Result->SetArrayField (TEXT("saved_packages"), SavedArray);
 	if (FailedCount > 0)
 	{
 		Result->SetArrayField(TEXT("failed_packages"), FailedArray);
 	}
+	if (bValidate)
+	{
+		Result->SetNumberField(TEXT("validation_errors"),   ValidationErrors);
+		Result->SetNumberField(TEXT("validation_warnings"), ValidationWarnings);
+	}
 	return FECACommandResult::Success(Result);
 }
 
diff --git a/Engine/Plugins/Experimental/ECABridge/Source/ECABridge/Public/Commands/ECASaveCommands.h b/Engine/Plugins/Experimental/ECABridge/Source/ECABridge/Public/Commands/ECASaveCommands.h
--- a/Engine/Plugins/Experimental/ECABridge/Source/ECABridge/Public/Commands/ECASaveCommands.h
+++ b/Engine/Plugins/Experimental/ECABridge/Source/ECABridge/Public/Commands/ECASaveCommands.h
@@ -47,8 +47,16 @@ class FECACommand_SaveAllDirty : public IECACommand
 {
 public:
 	virtual FString GetName() const override { return TEXT("save_all_dirty"); }
-	virtual FString GetDescription() const override { return TEXT("Save all modified assets to disk"); }
+	virtual FString GetDescription() const override { return TEXT("Save all modified assets to disk. Pass validate:true to also run UE Data Validation on each saved asset and surface any validator errors in the result."); }
 	virtual FString GetCategory() const override { return TEXT("Asset"); }
+
+	virtual TArray<FECACommandParam> GetParameters() const override
+	{
+		return {
+			{ TEXT("validate"), TEXT("bool"), TEXT("If true, run UE Data Validation on all saved assets after saving (same checks as Asset Check panel). Errors appear in 'validation_errors', warnings in 'validation_warnings' in the result."), false },
+		};
+	}
+
 	virtual FECACommandResult Execute(const TSharedPtr<FJsonObject>& Params) override;
 };
 
