[ECABridge] Widget designer screenshots and graceful editor shutdown

Adds two editor workflow tools:

1. take_gameplay_screenshot now supports 'widget_designer' as a target,
capturing the Widget Blueprint designer canvas directly. The Blueprint
editor must already be open. Uses exact window title matching to avoid
false positives when multiple WBPs are open.

2. close_editor: gracefully shuts down the Unreal Editor via
FPlatformMisc::RequestExit. By default saves all dirty assets first.
Pass save_dirty=false to skip saving and exit immediately.

#rnx

#Synced-CL 52344065

diff --git a/Engine/Plugins/Experimental/ECABridge/Source/ECABridge/Private/Commands/ECAEditorCommands.cpp b/Engine/Plugins/Experimental/ECABridge/Source/ECABridge/Private/Commands/ECAEditorCommands.cpp
--- a/Engine/Plugins/Experimental/ECABridge/Source/ECABridge/Private/Commands/ECAEditorCommands.cpp
+++ b/Engine/Plugins/Experimental/ECABridge/Source/ECABridge/Private/Commands/ECAEditorCommands.cpp
@@ -35,6 +35,7 @@ REGISTER_ECA_COMMAND(FECACommand_GetSelectedActors)
 
 REGISTER_ECA_COMMAND(FECACommand_TakeDepthScreenshot)
 REGISTER_ECA_COMMAND(FECACommand_TakeGameplayScreenshot)
+REGISTER_ECA_COMMAND(FECACommand_CloseEditor)
 REGISTER_ECA_COMMAND(FECACommand_RunConsoleCommand)
 REGISTER_ECA_COMMAND(FECACommand_GetCVar)
 REGISTER_ECA_COMMAND(FECACommand_GetLevelInfo)
@@ -525,10 +526,50 @@ FECACommandResult FECACommand_TakeGameplayScreenshot::Execute(const TSharedPtr<F
 
 		Source = TEXT("editor_viewport");
 	}
+	else if (Target == TEXT("widget_designer"))
+	{
+		// Capture the widget blueprint editor window for the specified widget_path.
+		FString WidgetPath;
+		GetStringParam(Params, TEXT("widget_path"), WidgetPath, false);
+		FString AssetName = FPaths::GetBaseFilename(WidgetPath);
+		if (AssetName.IsEmpty())
+		{
+			return FECACommandResult::Error(TEXT("widget_designer target requires widget_path"));
+		}
+
+		// Match the exact window title "AssetName Designer" to avoid false positives
+		// when one asset name is a substring of another (e.g. "WBP_Button" vs "WBP_ButtonGroup").
+		const FString ExpectedTitle = AssetName + TEXT(" Designer");
+		TSharedPtr<SWindow> DesignerWindow;
+		TArray<TSharedRef<SWindow>> AllWindows;
+		FSlateApplication::Get().GetAllVisibleWindowsOrdered(AllWindows);
+		for (const TSharedRef<SWindow>& Win : AllWindows)
+		{
+			if (Win->GetTitle().ToString() == ExpectedTitle)
+			{
+				DesignerWindow = Win;
+				break;
+			}
+		}
+		if (!DesignerWindow.IsValid())
+		{
+			return FECACommandResult::Error(FString::Printf(
+				TEXT("No widget designer window found with title '%s'. Open the WBP in the editor first."), *ExpectedTitle));
+		}
+		FIntVector Size;
+		bool bSuccess = FSlateApplication::Get().TakeScreenshot(DesignerWindow.ToSharedRef(), Bitmap, Size);
+		if (!bSuccess || Bitmap.Num() == 0)
+		{
+			return FECACommandResult::Error(TEXT("Failed to capture widget designer window."));
+		}
+		Width  = Size.X;
+		Height = Size.Y;
+		Source = TEXT("widget_designer");
+	}
 	else
 	{
 		return FECACommandResult::Error(FString::Printf(
-			TEXT("Unknown target '%s'. Valid targets: 'pie', 'editor_viewport', 'editor_window'."), *Target));
+			TEXT("Unknown target '%s'. Valid targets: 'pie', 'editor_viewport', 'editor_window', 'widget_designer'."), *Target));
 	}
 
 	if (Bitmap.Num() == 0)
@@ -849,3 +890,37 @@ FECACommandResult FECACommand_GetPlayState::Execute(const TSharedPtr<FJsonObject
 	
 	return FECACommandResult::Success(Result);
 }
+
+
+//------------------------------------------------------------------------------
+// CloseEditor
+//------------------------------------------------------------------------------
+FECACommandResult FECACommand_CloseEditor::Execute(const TSharedPtr<FJsonObject>& Params)
+{
+	bool bSaveDirty = true;
+	GetBoolParam(Params, TEXT("save_dirty"), bSaveDirty, false);
+
+	bool bSaveSucceeded = true;
+	if (bSaveDirty)
+	{
+		// SaveDirtyPackages(bSaveMapPackages, bSaveContentPackages) — 2-arg overload
+		bSaveSucceeded = UEditorLoadingAndSavingUtils::SaveDirtyPackages(true, true);
+	}
+
+	TSharedPtr<FJsonObject> Result = MakeResult();
+	Result->SetBoolField(TEXT("save_dirty"), bSaveDirty);
+	Result->SetBoolField(TEXT("save_succeeded"), bSaveSucceeded);
+
+	if (bSaveDirty && !bSaveSucceeded)
+	{
+		Result->SetStringField(TEXT("status"), TEXT("aborted_save_failed"));
+		return FECACommandResult::Error(TEXT("Save failed — editor exit aborted to prevent data loss."));
+	}
+
+	Result->SetStringField(TEXT("status"), TEXT("exit_requested"));
+
+	// FEngineLoop::AppExit / GIsRequestingExit — clean shutdown that respects the message loop.
+	FPlatformMisc::RequestExit(/*bForce=*/false);
+
+	return FECACommandResult::Success(Result);
+}
diff --git a/Engine/Plugins/Experimental/ECABridge/Source/ECABridge/Public/Commands/ECAEditorCommands.h b/Engine/Plugins/Experimental/ECABridge/Source/ECABridge/Public/Commands/ECAEditorCommands.h
--- a/Engine/Plugins/Experimental/ECABridge/Source/ECABridge/Public/Commands/ECAEditorCommands.h
+++ b/Engine/Plugins/Experimental/ECABridge/Source/ECABridge/Public/Commands/ECAEditorCommands.h
@@ -88,24 +88,26 @@ public:
 
 /**
  * Take a screenshot of the active viewport, game window, or editor window.
- * Supports three capture targets:
- *   "pie"             - Game window via FSlateApplication::TakeScreenshot (full composited with all UI)
- *   "editor_viewport" - Editor 3D viewport via GetViewportScreenShot (render target only, no Slate)
- *   "editor_window"   - Full editor window via FSlateApplication::TakeScreenshot (all panels, tabs, graphs)
+ * Supports four capture targets:
+ *   "pie"              - Game window via FSlateApplication::TakeScreenshot (full composited with all UI)
+ *   "editor_viewport"  - Editor 3D viewport via GetViewportScreenShot (render target only, no Slate)
+ *   "editor_window"    - Full editor window via FSlateApplication::TakeScreenshot (all panels, tabs, graphs)
+ *   "widget_designer"  - Widget Blueprint designer canvas (Blueprint editor must already be open; requires widget_path)
  * Auto-detects: if PIE is running uses "pie", otherwise uses "editor_viewport".
  */
 class FECACommand_TakeGameplayScreenshot : public IECACommand
 {
 public:
 	virtual FString GetName() const override { return TEXT("take_gameplay_screenshot"); }
-	virtual FString GetDescription() const override { return TEXT("Take a screenshot of the game viewport, editor viewport, or the full editor window. By default auto-detects: captures the PIE game window (with all Slate/UMG/HUD overlays) if playing, otherwise captures the editor 3D viewport. Use the 'target' parameter to explicitly choose what to capture, including 'editor_window' for the full editor UI (Blueprint graphs, widget hierarchies, Details panels, etc)."); }
+	virtual FString GetDescription() const override { return TEXT("Take a screenshot of the game viewport, editor viewport, full editor window, or Widget Blueprint designer canvas. By default auto-detects: captures the PIE game window (with all Slate/UMG/HUD overlays) if playing, otherwise captures the editor 3D viewport. Use the 'target' parameter to explicitly choose what to capture, including 'editor_window' for the full editor UI (Blueprint graphs, widget hierarchies, Details panels, etc) or 'widget_designer' to capture a Widget Blueprint's designer canvas (Blueprint editor must already be open)."); }
 	virtual FString GetCategory() const override { return TEXT("Editor"); }
 	
 	virtual TArray<FECACommandParam> GetParameters() const override
 	{
 		return {
 			{ TEXT("file_path"), TEXT("string"), TEXT("Path to save the screenshot PNG. If not specified, saves to project Saved/Screenshots/ with a timestamp filename."), false },
-			{ TEXT("target"), TEXT("string"), TEXT("What to capture: 'pie' (game window with UI overlays), 'editor_viewport' (3D viewport only), or 'editor_window' (full editor UI including all panels/tabs). If not specified, auto-detects: PIE window if playing, otherwise editor viewport."), false }
+			{ TEXT("target"), TEXT("string"), TEXT("What to capture: 'pie' (game window with UI overlays), 'editor_viewport' (3D viewport only), 'editor_window' (full editor UI including all panels/tabs), or 'widget_designer' (Widget Blueprint designer canvas — Blueprint editor must already be open, requires widget_path). If not specified, auto-detects: PIE window if playing, otherwise editor viewport."), false },
+			{ TEXT("widget_path"), TEXT("string"), TEXT("Asset path of the Widget Blueprint to capture. The Blueprint editor must already be open. Required when target is 'widget_designer'."), false }
 		};
 	}
 	
@@ -244,3 +246,26 @@ public:
 	
 	virtual FECACommandResult Execute(const TSharedPtr<FJsonObject>& Params) override;
 };
+
+/**
+ * Request the editor to exit gracefully. By default saves all dirty assets first,
+ * then calls FPlatformMisc::RequestExit(false) for a clean shutdown without the
+ * auto-restart behavior that 'quit' console command can trigger.
+ * Pass save_dirty=false to skip saving and exit immediately.
+ */
+class FECACommand_CloseEditor : public IECACommand
+{
+public:
+	virtual FString GetName() const override { return TEXT("close_editor"); }
+	virtual FString GetDescription() const override { return TEXT("Gracefully close the Unreal Editor. By default saves all dirty assets first, then requests a clean exit. Pass save_dirty=false to skip saving and exit immediately."); }
+	virtual FString GetCategory() const override { return TEXT("Editor"); }
+
+	virtual TArray<FECACommandParam> GetParameters() const override
+	{
+		return {
+			{ TEXT("save_dirty"), TEXT("bool"), TEXT("If true (default), saves all dirty packages before exiting"), false },
+		};
+	}
+
+	virtual FECACommandResult Execute(const TSharedPtr<FJsonObject>& Params) override;
+};
\ No newline at end of file
