[ECABridge] create_widget_from_spec improvements

create_widget_from_spec can now build complete Widget Blueprints in a single call --
including widget properties (Visibility, SizeBox overrides, Spacer sizes, etc.),
Blueprint Generated Class children from GameFeature plugins, and ViewModel assignments
from plugin modules. Previously, BPGC widget classes and plugin ViewModels were not
found, requiring manual follow-up calls. The tool also no longer crashes when the
target asset already exists.

set_widget_property gains a reflection fallback -- any UPROPERTY on any widget class
can now be set, not just the handful of explicitly supported types.

Both commands share a common property-setting helper (TrySetWidgetPropertyFromJson)
that uses typed UWidget API setters where they exist (to preserve side effects like
layout invalidation) and falls back to UE reflection for everything else.

#rnx
#Synced-CL 52369836

diff --git a/Engine/Plugins/Experimental/ECABridge/Source/ECABridge/Private/Commands/ECAWidgetTreeCommands.cpp b/Engine/Plugins/Experimental/ECABridge/Source/ECABridge/Private/Commands/ECAWidgetTreeCommands.cpp
--- a/Engine/Plugins/Experimental/ECABridge/Source/ECABridge/Private/Commands/ECAWidgetTreeCommands.cpp
+++ b/Engine/Plugins/Experimental/ECABridge/Source/ECABridge/Private/Commands/ECAWidgetTreeCommands.cpp
@@ -29,6 +29,7 @@
 #include "Animation/WidgetAnimation.h"
 #include "MovieScene.h"
 
+
 // Engine - UMG Slots
 #include "Components/PanelSlot.h"
 #include "Components/SlateWrapperTypes.h"
@@ -60,6 +61,304 @@
 #include "Components/PostBufferBlurUpdater.h"
 #include "Rendering/SlateRendererTypes.h"
 
+// Engine - Reflection (for set_widget_property generic fallback)
+#include "JsonObjectConverter.h"
+#include "UObject/FieldIterator.h"
+#include "UObject/UnrealType.h"
+
+/**
+ * Try to set a property on a UObject from a JSON value using UE reflection.
+ * Handles: bool, int32, float, double, FString, FName, FText, byte/enum (by name or value),
+ * FVector2D, FLinearColor, FVector, and complex structs via FJsonObjectConverter.
+ * @return true if the property was found and set successfully.
+ */
+static bool TrySetPropertyFromJson(UObject* Object, const FString& PropName, const TSharedPtr<FJsonValue>& PropValue)
+{
+	if (!Object || !PropValue.IsValid()) { return false; }
+
+	FProperty* Property = Object->GetClass()->FindPropertyByName(FName(*PropName));
+	if (!Property)
+	{
+		for (TFieldIterator<FProperty> It(Object->GetClass()); It; ++It)
+		{
+			if (It->GetName().Equals(PropName, ESearchCase::IgnoreCase))
+			{
+				Property = *It;
+				break;
+			}
+		}
+	}
+	if (!Property) { return false; }
+
+	void* ValuePtr = Property->ContainerPtrToValuePtr<void>(Object);
+
+	if (FBoolProperty* BoolProp = CastField<FBoolProperty>(Property))
+	{
+		bool bValue = false;
+		if (PropValue->TryGetBool(bValue)) { BoolProp->SetPropertyValue(ValuePtr, bValue); return true; }
+	}
+	else if (FIntProperty* IntProp = CastField<FIntProperty>(Property))
+	{
+		double NumVal = 0;
+		if (PropValue->TryGetNumber(NumVal)) { IntProp->SetPropertyValue(ValuePtr, (int32)NumVal); return true; }
+	}
+	else if (FFloatProperty* FloatProp = CastField<FFloatProperty>(Property))
+	{
+		double NumVal = 0;
+		if (PropValue->TryGetNumber(NumVal)) { FloatProp->SetPropertyValue(ValuePtr, (float)NumVal); return true; }
+	}
+	else if (FDoubleProperty* DoubleProp = CastField<FDoubleProperty>(Property))
+	{
+		double NumVal = 0;
+		if (PropValue->TryGetNumber(NumVal)) { DoubleProp->SetPropertyValue(ValuePtr, NumVal); return true; }
+	}
+	else if (FStrProperty* StrProp = CastField<FStrProperty>(Property))
+	{
+		FString StrVal;
+		if (PropValue->TryGetString(StrVal)) { StrProp->SetPropertyValue(ValuePtr, StrVal); return true; }
+	}
+	else if (FNameProperty* NameProp = CastField<FNameProperty>(Property))
+	{
+		FString StrVal;
+		if (PropValue->TryGetString(StrVal)) { NameProp->SetPropertyValue(ValuePtr, FName(*StrVal)); return true; }
+	}
+	else if (FTextProperty* TextProp = CastField<FTextProperty>(Property))
+	{
+		FString StrVal;
+		if (PropValue->TryGetString(StrVal)) { TextProp->SetPropertyValue(ValuePtr, FText::FromString(StrVal)); return true; }
+	}
+	else if (FByteProperty* ByteProp = CastField<FByteProperty>(Property))
+	{
+		if (ByteProp->Enum)
+		{
+			FString StrVal;
+			if (PropValue->TryGetString(StrVal))
+			{
+				int64 EnumValue = ByteProp->Enum->GetValueByNameString(StrVal);
+				if (EnumValue != INDEX_NONE) { ByteProp->SetPropertyValue(ValuePtr, (uint8)EnumValue); return true; }
+			}
+		}
+		double NumVal = 0;
+		if (PropValue->TryGetNumber(NumVal)) { ByteProp->SetPropertyValue(ValuePtr, (uint8)NumVal); return true; }
+	}
+	else if (FEnumProperty* EnumProp = CastField<FEnumProperty>(Property))
+	{
+		FString StrVal;
+		if (PropValue->TryGetString(StrVal))
+		{
+			int64 EnumValue = EnumProp->GetEnum()->GetValueByNameString(StrVal);
+			if (EnumValue != INDEX_NONE)
+			{
+				EnumProp->GetUnderlyingProperty()->SetIntPropertyValue(ValuePtr, EnumValue);
+				return true;
+			}
+		}
+		double NumVal = 0;
+		if (PropValue->TryGetNumber(NumVal))
+		{
+			EnumProp->GetUnderlyingProperty()->SetIntPropertyValue(ValuePtr, (int64)NumVal);
+			return true;
+		}
+	}
+	else if (FStructProperty* StructProp = CastField<FStructProperty>(Property))
+	{
+		const TSharedPtr<FJsonObject>* ObjVal = nullptr;
+		if (PropValue->TryGetObject(ObjVal) && ObjVal && ObjVal->IsValid())
+		{
+			if (StructProp->Struct == TBaseStructure<FVector2D>::Get())
+			{
+				FVector2D& Vec = *static_cast<FVector2D*>(ValuePtr);
+				(*ObjVal)->TryGetNumberField(TEXT("X"), Vec.X); (*ObjVal)->TryGetNumberField(TEXT("x"), Vec.X);
+				(*ObjVal)->TryGetNumberField(TEXT("Y"), Vec.Y); (*ObjVal)->TryGetNumberField(TEXT("y"), Vec.Y);
+				return true;
+			}
+			else if (StructProp->Struct == TBaseStructure<FLinearColor>::Get())
+			{
+				FLinearColor& Color = *static_cast<FLinearColor*>(ValuePtr);
+				double R = Color.R, G = Color.G, B = Color.B, A = Color.A;
+				(*ObjVal)->TryGetNumberField(TEXT("r"), R); (*ObjVal)->TryGetNumberField(TEXT("R"), R);
+				(*ObjVal)->TryGetNumberField(TEXT("g"), G); (*ObjVal)->TryGetNumberField(TEXT("G"), G);
+				(*ObjVal)->TryGetNumberField(TEXT("b"), B); (*ObjVal)->TryGetNumberField(TEXT("B"), B);
+				(*ObjVal)->TryGetNumberField(TEXT("a"), A); (*ObjVal)->TryGetNumberField(TEXT("A"), A);
+				Color = FLinearColor(R, G, B, A);
+				return true;
+			}
+			else if (StructProp->Struct == TBaseStructure<FVector>::Get())
+			{
+				FVector& Vec = *static_cast<FVector*>(ValuePtr);
+				(*ObjVal)->TryGetNumberField(TEXT("X"), Vec.X); (*ObjVal)->TryGetNumberField(TEXT("x"), Vec.X);
+				(*ObjVal)->TryGetNumberField(TEXT("Y"), Vec.Y); (*ObjVal)->TryGetNumberField(TEXT("y"), Vec.Y);
+				(*ObjVal)->TryGetNumberField(TEXT("Z"), Vec.Z); (*ObjVal)->TryGetNumberField(TEXT("z"), Vec.Z);
+				return true;
+			}
+			else if (FJsonObjectConverter::JsonObjectToUStruct(ObjVal->ToSharedRef(), StructProp->Struct, ValuePtr))
+			{
+				return true;
+			}
+		}
+	}
+
+	return false;
+}
+
+// Forward declaration — defined later in this file, but needed by TrySetWidgetPropertyFromJson
+static ESlateVisibility ParseVisibility(const FString& Value);
+
+/**
+ * Try to set a widget property using typed setter fast paths first, then reflection fallback.
+ * Uses UWidget API methods (SetVisibility, SetWidthOverride, SetSize, etc.) where they exist
+ * to correctly enable override flags and trigger rebuilds. Falls back to TrySetPropertyFromJson
+ * for properties without dedicated setters.
+ * @return true if the property was set successfully.
+ */
+static bool TrySetWidgetPropertyFromJson(UWidget* Widget, const FString& PropName, const TSharedPtr<FJsonValue>& PropValue)
+{
+	if (!Widget || !PropValue.IsValid()) { return false; }
+
+	// --- Universal UWidget properties (use API methods for correct side effects) ---
+	// When a fast path matches by name, it is authoritative — return false on type
+	// mismatch rather than falling through to reflection, which would bypass the
+	// API method's side effects (layout invalidation, rebuild, etc.).
+	if (PropName.Equals(TEXT("Visibility"), ESearchCase::IgnoreCase))
+	{
+		FString VisStr;
+		if (PropValue->TryGetString(VisStr)) { Widget->SetVisibility(ParseVisibility(VisStr)); return true; }
+		return false;
+	}
+	else if (PropName.Equals(TEXT("RenderOpacity"), ESearchCase::IgnoreCase))
+	{
+		double Val = 1.0;
+		if (PropValue->TryGetNumber(Val)) { Widget->SetRenderOpacity(Val); return true; }
+		return false;
+	}
+	else if (PropName.Equals(TEXT("IsEnabled"), ESearchCase::IgnoreCase))
+	{
+		bool bVal = true;
+		if (PropValue->TryGetBool(bVal)) { Widget->SetIsEnabled(bVal); return true; }
+		return false;
+	}
+	else if (PropName.Equals(TEXT("ToolTipText"), ESearchCase::IgnoreCase))
+	{
+		FString Str;
+		if (PropValue->TryGetString(Str)) { Widget->SetToolTipText(FText::FromString(Str)); return true; }
+		return false;
+	}
+	else if (PropName.Equals(TEXT("bIsVariable"), ESearchCase::IgnoreCase))
+	{
+		bool bVal = false;
+		if (PropValue->TryGetBool(bVal)) { Widget->bIsVariable = bVal; return true; }
+		return false;
+	}
+
+	// --- USizeBox overrides (SetXxxOverride enables bOverride_Xxx + triggers rebuild) ---
+	// Name-matched fast paths return false on type mismatch to prevent reflection bypass.
+	if (USizeBox* SizeBox = Cast<USizeBox>(Widget))
+	{
+		if (PropName.Equals(TEXT("WidthOverride"), ESearchCase::IgnoreCase))
+		{
+			double Val = 0; if (PropValue->TryGetNumber(Val)) { SizeBox->SetWidthOverride(Val); return true; }
+			return false;
+		}
+		else if (PropName.Equals(TEXT("HeightOverride"), ESearchCase::IgnoreCase))
+		{
+			double Val = 0; if (PropValue->TryGetNumber(Val)) { SizeBox->SetHeightOverride(Val); return true; }
+			return false;
+		}
+		else if (PropName.Equals(TEXT("MinDesiredWidth"), ESearchCase::IgnoreCase))
+		{
+			double Val = 0; if (PropValue->TryGetNumber(Val)) { SizeBox->SetMinDesiredWidth(Val); return true; }
+			return false;
+		}
+		else if (PropName.Equals(TEXT("MinDesiredHeight"), ESearchCase::IgnoreCase))
+		{
+			double Val = 0; if (PropValue->TryGetNumber(Val)) { SizeBox->SetMinDesiredHeight(Val); return true; }
+			return false;
+		}
+		else if (PropName.Equals(TEXT("MaxDesiredWidth"), ESearchCase::IgnoreCase))
+		{
+			double Val = 0; if (PropValue->TryGetNumber(Val)) { SizeBox->SetMaxDesiredWidth(Val); return true; }
+			return false;
+		}
+		else if (PropName.Equals(TEXT("MaxDesiredHeight"), ESearchCase::IgnoreCase))
+		{
+			double Val = 0; if (PropValue->TryGetNumber(Val)) { SizeBox->SetMaxDesiredHeight(Val); return true; }
+			return false;
+		}
+	}
+
+	// --- UTextBlock ---
+	if (UTextBlock* TextBlock = Cast<UTextBlock>(Widget))
+	{
+		if (PropName.Equals(TEXT("Text"), ESearchCase::IgnoreCase))
+		{
+			FString Str;
+			if (PropValue->TryGetString(Str)) { TextBlock->SetText(FText::FromString(Str)); return true; }
+			return false;
+		}
+		else if (PropName.Equals(TEXT("ColorAndOpacity"), ESearchCase::IgnoreCase))
+		{
+			const TSharedPtr<FJsonObject>* ColorObj = nullptr;
+			if (PropValue->TryGetObject(ColorObj) && ColorObj)
+			{
+				double R=1,G=1,B=1,A=1;
+				(*ColorObj)->TryGetNumberField(TEXT("r"),R); (*ColorObj)->TryGetNumberField(TEXT("g"),G);
+				(*ColorObj)->TryGetNumberField(TEXT("b"),B); (*ColorObj)->TryGetNumberField(TEXT("a"),A);
+				TextBlock->SetColorAndOpacity(FSlateColor(FLinearColor(R,G,B,A)));
+				return true;
+			}
+			return false;
+		}
+	}
+
+	// --- USpacer size (SetSize triggers rebuild) ---
+	if (USpacer* Spacer = Cast<USpacer>(Widget))
+	{
+		if (PropName.Equals(TEXT("Size"), ESearchCase::IgnoreCase) || PropName.Equals(TEXT("DesiredSize"), ESearchCase::IgnoreCase))
+		{
+			const TSharedPtr<FJsonObject>* SizeObj = nullptr;
+			if (PropValue->TryGetObject(SizeObj) && SizeObj && SizeObj->IsValid())
+			{
+				const FVector2D CurrentSize = Spacer->GetSize();
+				double X = CurrentSize.X, Y = CurrentSize.Y;
+				(*SizeObj)->TryGetNumberField(TEXT("X"), X); (*SizeObj)->TryGetNumberField(TEXT("x"), X);
+				(*SizeObj)->TryGetNumberField(TEXT("Y"), Y); (*SizeObj)->TryGetNumberField(TEXT("y"), Y);
+				Spacer->SetSize(FVector2D(X, Y));
+				return true;
+			}
+			double Val = 0;
+			if (PropValue->TryGetNumber(Val)) { Spacer->SetSize(FVector2D(Spacer->GetSize().X, Val)); return true; }
+			return false;
+		}
+	}
+
+	// --- UImage ---
+	if (UImage* Image = Cast<UImage>(Widget))
+	{
+		if (PropName.Equals(TEXT("ColorAndOpacity"), ESearchCase::IgnoreCase))
+		{
+			const TSharedPtr<FJsonObject>* ColorObj = nullptr;
+			if (PropValue->TryGetObject(ColorObj) && ColorObj)
+			{
+				double R=1,G=1,B=1,A=1;
+				(*ColorObj)->TryGetNumberField(TEXT("r"),R); (*ColorObj)->TryGetNumberField(TEXT("g"),G);
+				(*ColorObj)->TryGetNumberField(TEXT("b"),B); (*ColorObj)->TryGetNumberField(TEXT("a"),A);
+				Image->SetColorAndOpacity(FLinearColor(R,G,B,A));
+				return true;
+			}
+			return false;
+		}
+		else if (PropName.Equals(TEXT("Opacity"), ESearchCase::IgnoreCase))
+		{
+			double Val = 1.0;
+			if (PropValue->TryGetNumber(Val)) { Image->SetOpacity(Val); return true; }
+			return false;
+		}
+	}
+
+	// --- Reflection fallback for any other UPROPERTY ---
+	return TrySetPropertyFromJson(Widget, PropName, PropValue);
+}
+
 // Register all widget tree commands
 REGISTER_ECA_COMMAND(FECACommand_WrapWidget)
 REGISTER_ECA_COMMAND(FECACommand_UnwrapWidget)
@@ -77,6 +376,7 @@ REGISTER_ECA_COMMAND(FECACommand_CreateWidgetAnimation)
 
 REGISTER_ECA_COMMAND(FECACommand_SetWidgetDesignSize)
 
+
 //------------------------------------------------------------------------------
 // Shared Helper Functions
 //------------------------------------------------------------------------------
@@ -898,219 +1198,21 @@ FECACommandResult FECACommand_SetWidgetProperty::Execute(const TSharedPtr<FJsonO
 	TArray<FString> PropertiesSet;
 	TArray<FString> PropertiesSkipped;
 
-	// Process each property
+	// Process each property — delegates to shared TrySetWidgetPropertyFromJson which handles
+	// typed setter fast paths (Visibility, SizeBox, TextBlock, Spacer, Image) + reflection fallback.
 	for (const auto& Pair : Properties->Values)
 	{
 		const FString PropName(*Pair.Key);
 		const TSharedPtr<FJsonValue>& PropValue = Pair.Value;
 
-		bool bHandled = false;
-
-		// --- Universal UWidget properties ---
-
-		if (PropName.Equals(TEXT("Visibility"), ESearchCase::IgnoreCase))
-		{
-			FString VisStr;
-			if (PropValue->TryGetString(VisStr))
-			{
-				Widget->SetVisibility(ParseVisibility(VisStr));
-				bHandled = true;
-			}
-		}
-		else if (PropName.Equals(TEXT("RenderOpacity"), ESearchCase::IgnoreCase))
-		{
-			double Opacity = 1.0;
-			if (PropValue->TryGetNumber(Opacity))
-			{
-				Widget->SetRenderOpacity(Opacity);
-				bHandled = true;
-			}
-		}
-		else if (PropName.Equals(TEXT("IsEnabled"), ESearchCase::IgnoreCase))
-		{
-			bool bEnabled = true;
-			if (PropValue->TryGetBool(bEnabled))
-			{
-				Widget->SetIsEnabled(bEnabled);
-				bHandled = true;
-			}
-		}
-		else if (PropName.Equals(TEXT("ToolTipText"), ESearchCase::IgnoreCase))
-		{
-			FString TooltipStr;
-			if (PropValue->TryGetString(TooltipStr))
-			{
-				Widget->SetToolTipText(FText::FromString(TooltipStr));
-				bHandled = true;
-			}
-		}
-		else if (PropName.Equals(TEXT("bIsVariable"), ESearchCase::IgnoreCase))
-		{
-			bool bIsVariable = false;
-			if (PropValue->TryGetBool(bIsVariable))
-			{
-				Widget->bIsVariable = bIsVariable;
-				bHandled = true;
-			}
-		}
-
-		// --- USizeBox properties ---
-		else if (USizeBox* SizeBox = Cast<USizeBox>(Widget))
-		{
-			if (PropName.Equals(TEXT("WidthOverride"), ESearchCase::IgnoreCase))
-			{
-				double Val = 0;
-				if (PropValue->TryGetNumber(Val))
-				{
-					SizeBox->SetWidthOverride(Val);
-					bHandled = true;
-				}
-			}
-			else if (PropName.Equals(TEXT("HeightOverride"), ESearchCase::IgnoreCase))
-			{
-				double Val = 0;
-				if (PropValue->TryGetNumber(Val))
-				{
-					SizeBox->SetHeightOverride(Val);
-					bHandled = true;
-				}
-			}
-			else if (PropName.Equals(TEXT("MinDesiredWidth"), ESearchCase::IgnoreCase))
-			{
-				double Val = 0;
-				if (PropValue->TryGetNumber(Val))
-				{
-					SizeBox->SetMinDesiredWidth(Val);
-					bHandled = true;
-				}
-			}
-			else if (PropName.Equals(TEXT("MinDesiredHeight"), ESearchCase::IgnoreCase))
-			{
-				double Val = 0;
-				if (PropValue->TryGetNumber(Val))
-				{
-					SizeBox->SetMinDesiredHeight(Val);
-					bHandled = true;
-				}
-			}
-			else if (PropName.Equals(TEXT("MaxDesiredWidth"), ESearchCase::IgnoreCase))
-			{
-				double Val = 0;
-				if (PropValue->TryGetNumber(Val))
-				{
-					SizeBox->SetMaxDesiredWidth(Val);
-					bHandled = true;
-				}
-			}
-			else if (PropName.Equals(TEXT("MaxDesiredHeight"), ESearchCase::IgnoreCase))
-			{
-				double Val = 0;
-				if (PropValue->TryGetNumber(Val))
-				{
-					SizeBox->SetMaxDesiredHeight(Val);
-					bHandled = true;
-				}
-			}
-		}
-
-		// --- UTextBlock properties ---
-		else if (UTextBlock* TextBlock = Cast<UTextBlock>(Widget))
-		{
-			if (PropName.Equals(TEXT("Text"), ESearchCase::IgnoreCase))
-			{
-				FString TextStr;
-				if (PropValue->TryGetString(TextStr))
-				{
-					TextBlock->SetText(FText::FromString(TextStr));
-					bHandled = true;
-				}
-			}
-			else if (PropName.Equals(TEXT("ColorAndOpacity"), ESearchCase::IgnoreCase))
-			{
-				const TSharedPtr<FJsonObject>* ColorObj = nullptr;
-				if (PropValue->TryGetObject(ColorObj) && ColorObj)
-				{
-					double R = 1, G = 1, B = 1, A = 1;
-					(*ColorObj)->TryGetNumberField(TEXT("r"), R);
-					(*ColorObj)->TryGetNumberField(TEXT("g"), G);
-					(*ColorObj)->TryGetNumberField(TEXT("b"), B);
-					(*ColorObj)->TryGetNumberField(TEXT("a"), A);
-					TextBlock->SetColorAndOpacity(FSlateColor(FLinearColor(R, G, B, A)));
-					bHandled = true;
-				}
-			}
-		}
-
-		// --- USpacer properties ---
-		else if (USpacer* Spacer = Cast<USpacer>(Widget))
-		{
-			if (PropName.Equals(TEXT("Size"), ESearchCase::IgnoreCase) ||
-				PropName.Equals(TEXT("DesiredSize"), ESearchCase::IgnoreCase))
-			{
-				const TSharedPtr<FJsonObject>* SizeObj = nullptr;
-				if (PropValue->TryGetObject(SizeObj) && SizeObj && SizeObj->IsValid())
-				{
-					// Initialize from current size — omitted keys preserve the existing value (consistent with scalar path)
-					const FVector2D CurrentSize = Spacer->GetSize();
-					double X = CurrentSize.X, Y = CurrentSize.Y;
-					(*SizeObj)->TryGetNumberField(TEXT("X"), X);
-					(*SizeObj)->TryGetNumberField(TEXT("Y"), Y);
-					// Also accept lowercase
-					(*SizeObj)->TryGetNumberField(TEXT("x"), X);
-					(*SizeObj)->TryGetNumberField(TEXT("y"), Y);
-					Spacer->SetSize(FVector2D(X, Y));
-					bHandled = true;
-				}
-				else
-				{
-					// Accept a single number as the Y value (height)
-					double Val = 0;
-					if (PropValue->TryGetNumber(Val))
-					{
-						FVector2D Current = Spacer->GetSize();
-						Spacer->SetSize(FVector2D(Current.X, Val));
-						bHandled = true;
-					}
-				}
-			}
-		}
-
-		// --- UImage properties ---
-		else if (UImage* Image = Cast<UImage>(Widget))
-		{
-			if (PropName.Equals(TEXT("ColorAndOpacity"), ESearchCase::IgnoreCase))
-			{
-				const TSharedPtr<FJsonObject>* ColorObj = nullptr;
-				if (PropValue->TryGetObject(ColorObj) && ColorObj)
-				{
-					double R = 1, G = 1, B = 1, A = 1;
-					(*ColorObj)->TryGetNumberField(TEXT("r"), R);
-					(*ColorObj)->TryGetNumberField(TEXT("g"), G);
-					(*ColorObj)->TryGetNumberField(TEXT("b"), B);
-					(*ColorObj)->TryGetNumberField(TEXT("a"), A);
-					Image->SetColorAndOpacity(FLinearColor(R, G, B, A));
-					bHandled = true;
-				}
-			}
-			else if (PropName.Equals(TEXT("Opacity"), ESearchCase::IgnoreCase))
-			{
-				double Opacity = 1.0;
-				if (PropValue->TryGetNumber(Opacity))
-				{
-					Image->SetOpacity(Opacity);
-					bHandled = true;
-				}
-			}
-		}
-
-		if (bHandled)
+		if (TrySetWidgetPropertyFromJson(Widget, PropName, PropValue))
 		{
 			PropertiesSet.Add(PropName);
 		}
 		else
 		{
 			FString WidgetClass = Widget->GetClass()->GetName();
-			PropertiesSkipped.Add(FString::Printf(TEXT("%s (not supported on %s)"), *PropName, *WidgetClass));
+			PropertiesSkipped.Add(FString::Printf(TEXT("%s (not supported on %s, no matching UPROPERTY found)"), *PropName, *WidgetClass));
 		}
 	}
 
@@ -2214,19 +2316,10 @@ static void CreateWidgetsFromSpec(
 		bool bIsVariable = false;
 		Def->TryGetBoolField(TEXT("is_variable"), bIsVariable);
 
-		// Find or create widget class
-		UClass* Class = FindObject<UClass>(nullptr, *FString::Printf(TEXT("/Script/UMG.%s"), *WidgetClass));
-		if (!Class) Class = LoadObject<UClass>(nullptr, *FString::Printf(TEXT("/Script/UMG.%s"), *WidgetClass));
-		if (!Class) Class = FindObject<UClass>(nullptr, *(WidgetClass + TEXT("_C")));
-		if (!Class) Class = LoadObject<UClass>(nullptr, *(WidgetClass + TEXT("_C")));
-		// WBP blueprint class fallback: try loading as a Blueprint-generated UClass via asset path.
-		// Handles bare names like 'WBP_Foo_C' and content paths like '/Game/UI/WBP_Foo_C'.
-		if (!Class)
-		{
-			FString ClassPath = WidgetClass;
-			if (!ClassPath.EndsWith(TEXT("_C"))) ClassPath += TEXT("_C");
-			Class = LoadClass<UUserWidget>(nullptr, *ClassPath);
-		}
+		// Find widget class — uses the same resolution as add_child_widget,
+		// which searches engine classes, then all loaded UWidget subclasses via
+		// TObjectIterator. This correctly finds Blueprint Generated Classes (WBP_*_C).
+		UClass* Class = FindWidgetClass(WidgetClass);
 		if (!Class)
 		{
 			OutErrors.Add(FString::Printf(TEXT("Widget class not found: %s"), *WidgetClass));
@@ -2448,6 +2541,24 @@ static void CreateWidgetsFromSpec(
 			}
 		}
 
+		// Apply inline widget properties via the same logic as set_widget_property:
+		// typed setter fast paths for Visibility, SizeBox, Spacer (these use UWidget API
+		// methods that enable override flags, unlike raw reflection), then reflection fallback.
+		const TSharedPtr<FJsonObject>* PropsDef = nullptr;
+		if (Def->TryGetObjectField(TEXT("properties"), PropsDef) && PropsDef && PropsDef->IsValid())
+		{
+			for (const auto& Pair : (*PropsDef)->Values)
+			{
+				const FString PropName(Pair.Key);
+				const TSharedPtr<FJsonValue>& PropValue = Pair.Value;
+
+				if (!TrySetWidgetPropertyFromJson(NewWidget, PropName, PropValue))
+				{
+					OutErrors.Add(FString::Printf(TEXT("Widget '%s': property '%s' could not be set"), *WidgetName, *PropName));
+				}
+			}
+		}
+
 		// Recurse into children
 		const TArray<TSharedPtr<FJsonValue>>* ChildDefs = nullptr;
 		if (Def->TryGetArrayField(TEXT("children"), ChildDefs) && ChildDefs)
@@ -2484,10 +2595,30 @@ FECACommandResult FECACommand_CreateWidgetFromSpec::Execute(const TSharedPtr<FJs
 	{
 		return FECACommandResult::Error(TEXT("widget_path must be a full asset path with at least one '/' (e.g. /Game/UI/WBP_MyWidget)"));
 	}
-	const FString AssetName = WidgetPath.Mid(SlashIdx + 1);
+	// Strip fully-qualified object path suffix if present (e.g. /Game/UI/WBP_Foo.WBP_Foo → /Game/UI/WBP_Foo)
+	FString PackagePath = WidgetPath;
+	int32 DotIdx;
+	if (PackagePath.FindChar(TEXT('.'), DotIdx))
+	{
+		PackagePath = PackagePath.Left(DotIdx);
+	}
+	const int32 PackageSlashIdx = PackagePath.Find(TEXT("/"), ESearchCase::IgnoreCase, ESearchDir::FromEnd);
+	const FString AssetName = PackagePath.Mid(PackageSlashIdx + 1);
+
+	// Check if the asset already exists — prevents assert crash in CreateBlueprint
+	UObject* ExistingAsset = StaticFindObject(UWidgetBlueprint::StaticClass(), nullptr, *(PackagePath + TEXT(".") + AssetName));
+	if (!ExistingAsset)
+	{
+		ExistingAsset = LoadObject<UWidgetBlueprint>(nullptr, *(PackagePath + TEXT(".") + AssetName));
+	}
+	if (ExistingAsset)
+	{
+		return FECACommandResult::Error(FString::Printf(
+			TEXT("Asset already exists at: %s. Delete it first or use a different path."), *PackagePath));
+	}
 
 	// In UE, CreatePackage takes the FULL asset path (package name == asset path, e.g. /Game/UI/WBP_Foo).
-	UPackage* Package = CreatePackage(*WidgetPath);
+	UPackage* Package = CreatePackage(*PackagePath);
 	if (!Package)
 	{
 		return FECACommandResult::Error(FString::Printf(TEXT("Failed to create package: %s"), *WidgetPath));
@@ -2529,8 +2660,20 @@ FECACommandResult FECACommand_CreateWidgetFromSpec::Execute(const TSharedPtr<FJs
 				VMDef->TryGetStringField(TEXT("class"), VMClass);
 				if (VMClass.IsEmpty()) continue; // name is optional — defaults to class short name
 
+				// Search for ViewModel class — use TObjectIterator fallback for plugin classes
 				UClass* Class = FindObject<UClass>(nullptr, *VMClass);
-				if (!Class) Class = LoadObject<UClass>(nullptr, *VMClass);
+				if (!Class) Class = FindObject<UClass>(nullptr, *FString::Printf(TEXT("/Script/ModelViewViewModel.%s"), *VMClass));
+				if (!Class)
+				{
+					for (TObjectIterator<UClass> It; It; ++It)
+					{
+						if (It->GetName() == VMClass || It->GetAuthoredName() == VMClass)
+						{
+							Class = *It;
+							break;
+						}
+					}
+				}
 				if (!Class) { Errors.Add(FString::Printf(TEXT("ViewModel class not found: %s"), *VMClass)); continue; }
 
 				// FN_Main API registers the context under the class short name.
@@ -2772,3 +2915,5 @@ FECACommandResult FECACommand_SetWidgetDesignSize::Execute(const TSharedPtr<FJso
 	Result->SetStringField(TEXT("mode"), CurrentMode);
 	return FECACommandResult::Success(Result);
 }
+
+
diff --git a/Engine/Plugins/Experimental/ECABridge/Source/ECABridge/Public/Commands/ECAWidgetTreeCommands.h b/Engine/Plugins/Experimental/ECABridge/Source/ECABridge/Public/Commands/ECAWidgetTreeCommands.h
--- a/Engine/Plugins/Experimental/ECABridge/Source/ECABridge/Public/Commands/ECAWidgetTreeCommands.h
+++ b/Engine/Plugins/Experimental/ECABridge/Source/ECABridge/Public/Commands/ECAWidgetTreeCommands.h
@@ -110,7 +110,7 @@ class FECACommand_SetWidgetProperty : public IECACommand
 {
 public:
 	virtual FString GetName() const override { return TEXT("set_widget_property"); }
-	virtual FString GetDescription() const override { return TEXT("Set properties on a widget (Visibility, SizeBox overrides, Text, etc.). Use set_slot_property for layout/padding."); }
+	virtual FString GetDescription() const override { return TEXT("Set properties on a widget. Supports common properties (Visibility, RenderOpacity, IsEnabled, ToolTipText) and type-specific properties (SizeBox overrides, TextBlock Text, Spacer Size, Image Opacity). For any other UPROPERTY, uses UE reflection as fallback — works with bool, int, float, string, FName, FText, enums (by name), FVector2D, FLinearColor, FVector, and complex structs via JSON. Use set_slot_property for layout/padding."); }
 	virtual FString GetCategory() const override { return TEXT("UMG"); }
 
 	virtual TArray<FECACommandParam> GetParameters() const override
@@ -376,3 +376,4 @@ public:
 
 	virtual FECACommandResult Execute(const TSharedPtr<FJsonObject>& Params) override;
 };
+
