Enabling plugin in uproject. Removing explict cpp23 callout in Build.cs

#rnx
#Synced-CL 52355115

diff --git a/Engine/Plugins/Experimental/ECABridge/Source/ECABridge/ECABridge.Build.cs b/Engine/Plugins/Experimental/ECABridge/Source/ECABridge/ECABridge.Build.cs
--- a/Engine/Plugins/Experimental/ECABridge/Source/ECABridge/ECABridge.Build.cs
+++ b/Engine/Plugins/Experimental/ECABridge/Source/ECABridge/ECABridge.Build.cs
@@ -8,7 +8,6 @@ public class ECABridge : ModuleRules
 	public ECABridge(ReadOnlyTargetRules Target) : base(Target)
 	{
 		PCHUsage = PCHUsageMode.UseExplicitOrSharedPCHs;
-		CppStandard = CppStandardVersion.Cpp23;
 
 		PublicDependencyModuleNames.AddRange(new string[]
 		{
