ECABridge: don't run in commandlets...

#Synced-CL 52419852

diff --git a/Engine/Plugins/Experimental/ECABridge/Source/ECABridge/Private/ECABridge.cpp b/Engine/Plugins/Experimental/ECABridge/Source/ECABridge/Private/ECABridge.cpp
--- a/Engine/Plugins/Experimental/ECABridge/Source/ECABridge/Private/ECABridge.cpp
+++ b/Engine/Plugins/Experimental/ECABridge/Source/ECABridge/Private/ECABridge.cpp
@@ -10,6 +10,13 @@
 
 void UECABridge::Initialize()
 {
+	// Don't start the MCP server during commandlet runs (e.g. cook)
+	if (IsRunningCommandlet())
+	{
+		UE_LOG(LogTemp, Log, TEXT("[ECABridge] Skipping MCP server start — running as commandlet"));
+		return;
+	}
+
 	const UECABridgeSettings* Settings = UECABridgeSettings::Get();
 	
 	if (!Settings->bAutoStart)
