Fix ECABridge MCP tool schema emitting invalid JSON Schema type "bool" instead of "boolean"

Two commands (close_editor, save_all_dirty) declared parameters with type "bool", which is not a valid JSON Schema type. The BuildToolDefinitions() type-mapping code handled "float", "int", etc. but not "bool". This caused the Anthropic API to reject the tool list with a 400 error when the ECABridge plugin was activated.

- Add "bool" -> "boolean" mapping in BuildToolDefinitions() for future-proofing
- Fix save_dirty param in ECAEditorCommands.h to use "boolean"
- Fix validate param in ECASaveCommands.h to use "boolean"

#tests Relaunched editor and enabled ECA Bridge MCP in Claude Code, with a basic investigation task to see that all ECA Bridge functionality worked within the session.
#Synced-CL 52405917

#Synced-CL 52405917

diff --git a/Engine/Plugins/Experimental/ECABridge/Source/ECABridge/Private/ECAMCPServer.cpp b/Engine/Plugins/Experimental/ECABridge/Source/ECABridge/Private/ECAMCPServer.cpp
--- a/Engine/Plugins/Experimental/ECABridge/Source/ECABridge/Private/ECAMCPServer.cpp
+++ b/Engine/Plugins/Experimental/ECABridge/Source/ECABridge/Private/ECAMCPServer.cpp
@@ -462,39 +462,30 @@ TArray<TSharedPtr<FJsonValue>> FECAMCPServer::BuildToolDefinitions()
 		for (const FECACommandParam& Param : Command->GetParameters())
 		{
 			TSharedPtr<FJsonObject> PropDef = MakeShared<FJsonObject>();
-			
-			// Map parameter types to valid JSON Schema types
-			// Valid types: string, number, integer, boolean, object, array, null
-			FString JsonSchemaType = Param.Type;
+
+			// map non-standard param types to valid JSON Schema types
+			static const TMap<FString, FString> TypeFixups = {
+				{ TEXT("bool"),      TEXT("boolean") },
+				{ TEXT("float"),     TEXT("number")  },
+				{ TEXT("double"),    TEXT("number")  },
+				{ TEXT("int"),       TEXT("integer") },
+				{ TEXT("int32"),     TEXT("integer") },
+				{ TEXT("int64"),     TEXT("integer") },
+				{ TEXT("vector"),    TEXT("object")  },
+				{ TEXT("rotator"),   TEXT("object")  },
+				{ TEXT("transform"), TEXT("object")  }
+			};
+
 			if (Param.Type == TEXT("any"))
 			{
 				// "any" is not valid in JSON Schema - use a union or just string
 				// We'll accept any JSON value by not specifying type constraints
 				PropDef->SetStringField(TEXT("description"), Param.Description + TEXT(" (accepts any JSON value)"));
 			}
-			else if (Param.Type == TEXT("vector") || Param.Type == TEXT("rotator") || Param.Type == TEXT("transform"))
-			{
-				// These are objects with specific structure
-				JsonSchemaType = TEXT("object");
-				PropDef->SetStringField(TEXT("type"), JsonSchemaType);
-				PropDef->SetStringField(TEXT("description"), Param.Description);
-			}
-			else if (Param.Type == TEXT("float") || Param.Type == TEXT("double"))
-			{
-				JsonSchemaType = TEXT("number");
-				PropDef->SetStringField(TEXT("type"), JsonSchemaType);
-				PropDef->SetStringField(TEXT("description"), Param.Description);
-			}
-			else if (Param.Type == TEXT("int") || Param.Type == TEXT("int32") || Param.Type == TEXT("int64"))
-			{
-				JsonSchemaType = TEXT("integer");
-				PropDef->SetStringField(TEXT("type"), JsonSchemaType);
-				PropDef->SetStringField(TEXT("description"), Param.Description);
-			}
 			else
 			{
-				// Standard JSON Schema types: string, number, integer, boolean, object, array
-				PropDef->SetStringField(TEXT("type"), JsonSchemaType);
+				const FString* FixedType = TypeFixups.Find(Param.Type);
+				PropDef->SetStringField(TEXT("type"), FixedType ? *FixedType : Param.Type);
 				PropDef->SetStringField(TEXT("description"), Param.Description);
 			}
 			
diff --git a/Engine/Plugins/Experimental/ECABridge/Source/ECABridge/Public/Commands/ECAEditorCommands.h b/Engine/Plugins/Experimental/ECABridge/Source/ECABridge/Public/Commands/ECAEditorCommands.h
--- a/Engine/Plugins/Experimental/ECABridge/Source/ECABridge/Public/Commands/ECAEditorCommands.h
+++ b/Engine/Plugins/Experimental/ECABridge/Source/ECABridge/Public/Commands/ECAEditorCommands.h
@@ -263,7 +263,7 @@ public:
 	virtual TArray<FECACommandParam> GetParameters() const override
 	{
 		return {
-			{ TEXT("save_dirty"), TEXT("bool"), TEXT("If true (default), saves all dirty packages before exiting"), false },
+			{ TEXT("save_dirty"), TEXT("boolean"), TEXT("If true (default), saves all dirty packages before exiting"), false },
 		};
 	}
 
diff --git a/Engine/Plugins/Experimental/ECABridge/Source/ECABridge/Public/Commands/ECASaveCommands.h b/Engine/Plugins/Experimental/ECABridge/Source/ECABridge/Public/Commands/ECASaveCommands.h
--- a/Engine/Plugins/Experimental/ECABridge/Source/ECABridge/Public/Commands/ECASaveCommands.h
+++ b/Engine/Plugins/Experimental/ECABridge/Source/ECABridge/Public/Commands/ECASaveCommands.h
@@ -53,7 +53,7 @@ public:
 	virtual TArray<FECACommandParam> GetParameters() const override
 	{
 		return {
-			{ TEXT("validate"), TEXT("bool"), TEXT("If true, run UE Data Validation on all saved assets after saving (same checks as Asset Check panel). Errors appear in 'validation_errors', warnings in 'validation_warnings' in the result."), false },
+			{ TEXT("validate"), TEXT("boolean"), TEXT("If true, run UE Data Validation on all saved assets after saving (same checks as Asset Check panel). Errors appear in 'validation_errors', warnings in 'validation_warnings' in the result."), false },
 		};
 	}
 
