[ECABridge] save_and_compile_blueprint: compile error reporting

Collects compile error messages from graph nodes after compilation.
Iterates UbergraphPages and FunctionGraphs, finds nodes with
bHasCompilerMessage and ErrorType <= Error, and returns them in a
compile_errors JSON array in the response.

Note: MVVM compiler errors go through LogBlueprint, not graph nodes.
These can be read via ToolsetRegistry.LogsToolset.GetLogEntries.

#rnx
#Synced-CL 52369836

diff --git a/Engine/Plugins/Experimental/ECABridge/Source/ECABridge/Private/Commands/ECASaveCommands.cpp b/Engine/Plugins/Experimental/ECABridge/Source/ECABridge/Private/Commands/ECASaveCommands.cpp
--- a/Engine/Plugins/Experimental/ECABridge/Source/ECABridge/Private/Commands/ECASaveCommands.cpp
+++ b/Engine/Plugins/Experimental/ECABridge/Source/ECABridge/Private/Commands/ECASaveCommands.cpp
@@ -5,6 +5,9 @@
 #include "Editor.h"
 #include "EditorValidatorSubsystem.h"
 #include "Engine/Blueprint.h"
+#include "EdGraph/EdGraph.h"
+#include "EdGraph/EdGraphNode.h"
+#include "Kismet2/BlueprintEditorUtils.h"
 #include "Kismet2/KismetEditorUtilities.h"
 #include "Misc/PackageName.h"
 #include "UObject/Package.h"
@@ -281,6 +284,29 @@ FECACommandResult FECACommand_SaveAndCompileBlueprint::Execute(const TSharedPtr<
 	FKismetEditorUtilities::CompileBlueprint(Blueprint);
 	bool bHasErrors = Blueprint->Status == BS_Error;
 
+	// Collect compile error messages from graph nodes
+	TArray<FString> CompileErrors;
+	if (bHasErrors)
+	{
+		auto CollectErrors = [&CompileErrors](const TArray<TObjectPtr<UEdGraph>>& Graphs)
+		{
+			for (UEdGraph* Graph : Graphs)
+			{
+				if (!Graph) continue;
+				for (UEdGraphNode* Node : Graph->Nodes)
+				{
+					if (Node && Node->bHasCompilerMessage && Node->ErrorType <= EMessageSeverity::Error)
+					{
+						CompileErrors.Add(FString::Printf(TEXT("[%s] %s: %s"),
+							*Graph->GetName(), *Node->GetNodeTitle(ENodeTitleType::ListView).ToString(), *Node->ErrorMsg));
+					}
+				}
+			}
+		};
+		CollectErrors(Blueprint->UbergraphPages);
+		CollectErrors(Blueprint->FunctionGraphs);
+	}
+
 	// Step 2: Save (always — even with compile errors, persist WIP)
 	UPackage* Package = Blueprint->GetPackage();
 	TArray<UPackage*> PackagesToSave;
@@ -290,6 +316,12 @@ FECACommandResult FECACommand_SaveAndCompileBlueprint::Execute(const TSharedPtr<
 	TSharedPtr<FJsonObject> Result = MakeResult();
 	Result->SetBoolField(TEXT("compiled"), true);
 	Result->SetBoolField(TEXT("has_errors"), bHasErrors);
+	if (!CompileErrors.IsEmpty())
+	{
+		TArray<TSharedPtr<FJsonValue>> ErrArr;
+		for (const FString& E : CompileErrors) ErrArr.Add(MakeShared<FJsonValueString>(E));
+		Result->SetArrayField(TEXT("compile_errors"), ErrArr);
+	}
 	Result->SetStringField(TEXT("status"), bHasErrors ? TEXT("Error") : TEXT("Success"));
 	Result->SetStringField(TEXT("blueprint_path"), BlueprintPath);
 	Result->SetStringField(TEXT("package_name"), Package->GetName());
