[ECABridge] add_viewmodel_binding: FProperty parameter fallback for conversion functions

GetConversionFunctionPins returns 0 pins for BP Function Library UFunctions
because TryGetArgumentsForConversionFunction fails validation. Fall back to
iterating the UFunction's FProperty parameters via TFieldIterator and
constructing FMVVMBlueprintPinId from the parameter FNames directly.

Also skips IsValidConversionFunction pre-check when the destination is a
UFunction (setter-style binding) since the engine's validation doesn't handle
this case correctly for BP library functions.

#rnx
#Synced-CL 52369836

diff --git a/Engine/Plugins/Experimental/ECABridge/Source/ECABridge/Private/Commands/ECAMVVMCommands.cpp b/Engine/Plugins/Experimental/ECABridge/Source/ECABridge/Private/Commands/ECAMVVMCommands.cpp
--- a/Engine/Plugins/Experimental/ECABridge/Source/ECABridge/Private/Commands/ECAMVVMCommands.cpp
+++ b/Engine/Plugins/Experimental/ECABridge/Source/ECABridge/Private/Commands/ECAMVVMCommands.cpp
@@ -9,7 +9,9 @@
 #include "Components/TextBlock.h"
 #include "Components/Widget.h"
 #include "EdGraph/EdGraphPin.h"
+#include "EdGraphSchema_K2.h"
 #include "K2Node.h"
+#include "K2Node_CallFunction.h"
 #include "Types/MVVMConversionFunctionValue.h"
 #include "Editor.h"
 #include "Kismet/BlueprintFunctionLibrary.h"
@@ -1212,11 +1214,24 @@ FECACommandResult FECACommand_AddViewmodelBinding::Execute(const TSharedPtr<FJso
 			SrcPath.SetViewModelId(SrcVMContext->GetViewModelId());
 			SrcPath.SetPropertyPath(WidgetBP, SrcField);
 
-			// Get ALL conversion function input pins (unfiltered) and index by SourceIdx.
-			// Passing nullptr returns all pins regardless of type, which is correct for
-			// heterogeneous-input functions (e.g. bool + int) where type-filtered results
-			// would be a subset and indexing by SourceIdx would pick the wrong pin.
+			// Get conversion function input pins. Try the standard API first.
+			// If that returns 0 pins (common with BP Function Library UFunctions whose
+			// TryGetArgumentsForConversionFunction fails), fall back to constructing
+			// pin IDs from the UFunction's FProperty parameters directly.
 			TArray<FMVVMBlueprintPinId> InputPins = MVVMSubsystem->GetConversionFunctionPins(WidgetBP, ConvValue, nullptr);
+			if (InputPins.Num() == 0)
+			{
+				// Fallback: construct pin IDs from the UFunction's parameter FProperties.
+				// This works for any UFunction regardless of how the MVVM subsystem validates it.
+				for (TFieldIterator<FProperty> ParamIt(ConvFunc); ParamIt && ParamIt->HasAnyPropertyFlags(CPF_Parm); ++ParamIt)
+				{
+					if (!ParamIt->HasAnyPropertyFlags(CPF_ReturnParm))
+					{
+						TArray<FName> Names = { ParamIt->GetFName() };
+						InputPins.Add(FMVVMBlueprintPinId(MoveTemp(Names)));
+					}
+				}
+			}
 			if (InputPins.IsValidIndex(SourceIdx))
 			{
 				MVVMSubsystem->SetPathForConversionFunctionArgument(WidgetBP, NewBinding, InputPins[SourceIdx], SrcPath, true);
