[ECABridge] create_widget_from_spec: Asset Registry fallback, inline variables/functions, CDO properties

FindWidgetClass now falls back to AssetRegistry search when TObjectIterator fails.
This loads unloaded GameFeature plugin widget BPs on demand, enabling
create_widget_from_spec to create GFP child widgets without pre-loading them.

Also adds inline support for variables, functions (with parameters), and
CDO properties directly in the spec JSON. Variables and functions are created
before the final compile so they exist on the generated class CDO.

#rnx
#Synced-CL 52369836

diff --git a/Engine/Plugins/Experimental/ECABridge/Source/ECABridge/Private/Commands/ECAWidgetTreeCommands.cpp b/Engine/Plugins/Experimental/ECABridge/Source/ECABridge/Private/Commands/ECAWidgetTreeCommands.cpp
--- a/Engine/Plugins/Experimental/ECABridge/Source/ECABridge/Private/Commands/ECAWidgetTreeCommands.cpp
+++ b/Engine/Plugins/Experimental/ECABridge/Source/ECABridge/Private/Commands/ECAWidgetTreeCommands.cpp
@@ -44,6 +44,9 @@
 #include "FileHelpers.h"
 #include "Kismet2/BlueprintEditorUtils.h"
 #include "Kismet2/KismetEditorUtilities.h"
+#include "K2Node_FunctionEntry.h"
+#include "EdGraphSchema_K2.h"
+#include "AssetRegistry/AssetRegistryModule.h"
 #include "MVVMBlueprintView.h"
 #include "MVVMBlueprintViewBinding.h"
 #include "MVVMBlueprintViewModelContext.h"
@@ -514,18 +517,68 @@ static UClass* FindWidgetClass(const FString& ClassName)
 	}
 
 	// Search all loaded UWidget subclasses (covers project/plugin types like CommonTextBlock, WBP_*_C)
+	// Try both the exact name and with _C suffix (Blueprint Generated Class convention)
+	FString NameWithC = ClassName.EndsWith(TEXT("_C")) ? ClassName : ClassName + TEXT("_C");
 	for (TObjectIterator<UClass> It; It; ++It)
 	{
 		UClass* Class = *It;
 		if (Class->IsChildOf(UWidget::StaticClass()) && !Class->HasAnyClassFlags(CLASS_Abstract))
 		{
-			if (Class->GetName().Equals(ClassName, ESearchCase::IgnoreCase))
+			const FString& Name = Class->GetName();
+			if (Name.Equals(ClassName, ESearchCase::IgnoreCase) || Name.Equals(NameWithC, ESearchCase::IgnoreCase))
 			{
 				return Class;
 			}
 		}
 	}
 
+	// Asset Registry fallback: find unloaded Blueprint assets by name and load them.
+	// This handles GameFeature plugin widgets that aren't in memory yet.
+	{
+		FAssetRegistryModule& AssetRegistryModule = FModuleManager::LoadModuleChecked<FAssetRegistryModule>("AssetRegistry");
+		IAssetRegistry& AssetRegistry = AssetRegistryModule.Get();
+
+		// Search for a Blueprint asset matching the class name (without _C suffix)
+		FString BaseName = ClassName;
+		if (BaseName.EndsWith(TEXT("_C")))
+		{
+			BaseName = BaseName.LeftChop(2);
+		}
+
+		TArray<FAssetData> FoundAssets;
+		AssetRegistry.GetAssetsByClass(UBlueprint::StaticClass()->GetClassPathName(), FoundAssets);
+
+		for (const FAssetData& Asset : FoundAssets)
+		{
+			if (Asset.AssetName.ToString().Equals(BaseName, ESearchCase::IgnoreCase))
+			{
+				// Load the Blueprint and get its generated class
+				UBlueprint* BP = Cast<UBlueprint>(Asset.GetAsset());
+				if (BP && BP->GeneratedClass && BP->GeneratedClass->IsChildOf(UWidget::StaticClass()))
+				{
+					return BP->GeneratedClass;
+				}
+			}
+		}
+
+		// Also search WidgetBlueprint — GetAssetsByClass uses exact class match,
+		// so the UBlueprint pass above does NOT return UWidgetBlueprint assets.
+		FoundAssets.Reset();
+		AssetRegistry.GetAssetsByClass(UWidgetBlueprint::StaticClass()->GetClassPathName(), FoundAssets);
+
+		for (const FAssetData& Asset : FoundAssets)
+		{
+			if (Asset.AssetName.ToString().Equals(BaseName, ESearchCase::IgnoreCase))
+			{
+				UBlueprint* BP = Cast<UBlueprint>(Asset.GetAsset());
+				if (BP && BP->GeneratedClass && BP->GeneratedClass->IsChildOf(UWidget::StaticClass()))
+				{
+					return BP->GeneratedClass;
+				}
+			}
+		}
+	}
+
 	return nullptr;
 }
 
@@ -2714,7 +2767,6 @@ FECACommandResult FECACommand_CreateWidgetFromSpec::Execute(const TSharedPtr<FJs
 	// Re-compile after CDO changes so MVVM GeneratedClass is fully up to date for self-bindings.
 	FKismetEditorUtilities::CompileBlueprint(WidgetBP);
 
-
 	// Add MVVM bindings
 	const TArray<TSharedPtr<FJsonValue>>* BindingDefs = nullptr;
 	if (Params->TryGetArrayField(TEXT("bindings"), BindingDefs) && BindingDefs)
@@ -2818,9 +2870,149 @@ FECACommandResult FECACommand_CreateWidgetFromSpec::Execute(const TSharedPtr<FJs
 		}
 	}
 
+	// Add Blueprint variables from spec (must be before final compile so they exist on the CDO)
+	const TArray<TSharedPtr<FJsonValue>>* VarDefs = nullptr;
+	if (Params->TryGetArrayField(TEXT("variables"), VarDefs) && VarDefs)
+	{
+		for (const TSharedPtr<FJsonValue>& VarVal : *VarDefs)
+		{
+			const TSharedPtr<FJsonObject>* VarObj = nullptr;
+			if (!VarVal->TryGetObject(VarObj) || !VarObj) continue;
+
+			FString VarName, VarType;
+			(*VarObj)->TryGetStringField(TEXT("name"), VarName);
+			(*VarObj)->TryGetStringField(TEXT("type"), VarType);
+			if (VarName.IsEmpty() || VarType.IsEmpty()) continue;
+
+			// Map type string to FEdGraphPinType
+			FEdGraphPinType PinType;
+			if (VarType.Equals(TEXT("bool"), ESearchCase::IgnoreCase))
+				PinType.PinCategory = UEdGraphSchema_K2::PC_Boolean;
+			else if (VarType.Equals(TEXT("int"), ESearchCase::IgnoreCase) || VarType.Equals(TEXT("int32"), ESearchCase::IgnoreCase))
+				PinType.PinCategory = UEdGraphSchema_K2::PC_Int;
+			else if (VarType.Equals(TEXT("float"), ESearchCase::IgnoreCase))
+				PinType.PinCategory = UEdGraphSchema_K2::PC_Real;
+			else if (VarType.Equals(TEXT("string"), ESearchCase::IgnoreCase))
+				PinType.PinCategory = UEdGraphSchema_K2::PC_String;
+			else if (VarType.Equals(TEXT("name"), ESearchCase::IgnoreCase))
+				PinType.PinCategory = UEdGraphSchema_K2::PC_Name;
+			else if (VarType.Equals(TEXT("text"), ESearchCase::IgnoreCase))
+				PinType.PinCategory = UEdGraphSchema_K2::PC_Text;
+			else if (VarType.Equals(TEXT("object"), ESearchCase::IgnoreCase))
+				PinType.PinCategory = UEdGraphSchema_K2::PC_Object;
+			else
+				PinType.PinCategory = UEdGraphSchema_K2::PC_Object; // default fallback
+
+			if (!FBlueprintEditorUtils::AddMemberVariable(WidgetBP, FName(*VarName), PinType))
+			{
+				Errors.Add(FString::Printf(TEXT("variables: failed to add '%s' (name collision or invalid type)"), *VarName));
+			}
+		}
+	}
+
+	// Add Blueprint functions from spec
+	const TArray<TSharedPtr<FJsonValue>>* FuncDefs = nullptr;
+	if (Params->TryGetArrayField(TEXT("functions"), FuncDefs) && FuncDefs)
+	{
+		for (const TSharedPtr<FJsonValue>& FuncVal : *FuncDefs)
+		{
+			const TSharedPtr<FJsonObject>* FuncObj = nullptr;
+			if (!FuncVal->TryGetObject(FuncObj) || !FuncObj) continue;
+
+			FString FuncName;
+			(*FuncObj)->TryGetStringField(TEXT("name"), FuncName);
+			if (FuncName.IsEmpty()) continue;
+
+			// Create function graph
+			UEdGraph* NewGraph = FBlueprintEditorUtils::CreateNewGraph(
+				WidgetBP, FName(*FuncName), UEdGraph::StaticClass(), UEdGraphSchema_K2::StaticClass());
+			if (NewGraph)
+			{
+				FBlueprintEditorUtils::AddFunctionGraph<UFunction>(WidgetBP, NewGraph, true, nullptr);
+
+				// Add parameters
+				const TArray<TSharedPtr<FJsonValue>>* ParamDefs = nullptr;
+				if ((*FuncObj)->TryGetArrayField(TEXT("params"), ParamDefs) && ParamDefs)
+				{
+					UK2Node_FunctionEntry* EntryNode = nullptr;
+					for (UEdGraphNode* Node : NewGraph->Nodes)
+					{
+						if (UK2Node_FunctionEntry* Found = Cast<UK2Node_FunctionEntry>(Node))
+						{
+							EntryNode = Found;
+							break;
+						}
+					}
+
+					if (EntryNode)
+					{
+						for (const TSharedPtr<FJsonValue>& ParamVal : *ParamDefs)
+						{
+							const TSharedPtr<FJsonObject>* ParamObj = nullptr;
+							if (!ParamVal->TryGetObject(ParamObj) || !ParamObj) continue;
+
+							FString ParamName, ParamType;
+							(*ParamObj)->TryGetStringField(TEXT("name"), ParamName);
+							(*ParamObj)->TryGetStringField(TEXT("type"), ParamType);
+							if (ParamName.IsEmpty()) continue;
+
+							FEdGraphPinType ParamPinType;
+							if (ParamType.Equals(TEXT("bool"), ESearchCase::IgnoreCase))
+								ParamPinType.PinCategory = UEdGraphSchema_K2::PC_Boolean;
+							else if (ParamType.Equals(TEXT("int"), ESearchCase::IgnoreCase) || ParamType.Equals(TEXT("int32"), ESearchCase::IgnoreCase))
+								ParamPinType.PinCategory = UEdGraphSchema_K2::PC_Int;
+							else if (ParamType.Equals(TEXT("float"), ESearchCase::IgnoreCase))
+								ParamPinType.PinCategory = UEdGraphSchema_K2::PC_Real;
+							else if (ParamType.Equals(TEXT("string"), ESearchCase::IgnoreCase))
+								ParamPinType.PinCategory = UEdGraphSchema_K2::PC_String;
+							else if (ParamType.Equals(TEXT("name"), ESearchCase::IgnoreCase))
+								ParamPinType.PinCategory = UEdGraphSchema_K2::PC_Name;
+							else if (ParamType.Equals(TEXT("text"), ESearchCase::IgnoreCase))
+								ParamPinType.PinCategory = UEdGraphSchema_K2::PC_Text;
+							else
+								ParamPinType.PinCategory = UEdGraphSchema_K2::PC_Object;
+
+							if (!EntryNode->CreateUserDefinedPin(FName(*ParamName), ParamPinType, EGPD_Output))
+							{
+								Errors.Add(FString::Printf(TEXT("functions: failed to add param '%s' to '%s' (duplicate or invalid)"), *ParamName, *FuncName));
+							}
+						}
+					}
+				}
+			}
+		}
+	}
+
 	// Re-compile after MVVM bindings are set up, then save
 	EnsureAllWidgetGuids(WidgetBP);
 	FKismetEditorUtilities::CompileBlueprint(WidgetBP);
+
+	// Set CDO properties AFTER final compile — BP variables (LaneCount, etc.) only exist
+	// on the generated class CDO after the variables are added and the BP is fully compiled.
+	const TSharedPtr<FJsonObject>* CdoPropsDef = nullptr;
+	if (Params->TryGetObjectField(TEXT("cdo_properties"), CdoPropsDef) && CdoPropsDef && CdoPropsDef->IsValid())
+	{
+		UObject* CDO = WidgetBP->GeneratedClass ? WidgetBP->GeneratedClass->GetDefaultObject(true) : nullptr;
+		if (CDO)
+		{
+			CDO->Modify();
+			for (const auto& Pair : (*CdoPropsDef)->Values)
+			{
+				const FString PropName(Pair.Key);
+				if (!TrySetPropertyFromJson(CDO, PropName, Pair.Value))
+				{
+					Errors.Add(FString::Printf(TEXT("cdo_properties: '%s' could not be set on CDO"), *PropName));
+				}
+			}
+			// Re-compile so CDO changes are serialized
+			FKismetEditorUtilities::CompileBlueprint(WidgetBP);
+		}
+		else
+		{
+			Errors.Add(TEXT("cdo_properties: CDO unavailable after compile"));
+		}
+	}
+
 	UEditorLoadingAndSavingUtils::SavePackages({ WidgetBP->GetOutermost() }, false);
 
 	// Result
