[ECABridge] Add CPF_Parm fallback for SetViewmodelBindingArgument conversion function pins

SetViewmodelBindingArgument's GetConversionFunctionPins() can return empty for
BP Function Library UFunctions whose TryGetArgumentsForConversionFunction fails.
This is the same issue that was fixed for AddViewmodelBinding in CL 52483853.

Adds the same CPF_Parm FProperty iteration fallback: when GetConversionFunctionPins()
returns empty but we have a resolved UFunction, iterate its FProperties with CPF_Parm
guard (skipping CPF_ReturnParm) to construct FMVVMBlueprintPinId entries manually.

#rnx
#Synced-CL 52509544

diff --git a/Engine/Plugins/Experimental/ECABridge/Source/ECABridge/Private/Commands/ECAMVVMCommands.cpp b/Engine/Plugins/Experimental/ECABridge/Source/ECABridge/Private/Commands/ECAMVVMCommands.cpp
--- a/Engine/Plugins/Experimental/ECABridge/Source/ECABridge/Private/Commands/ECAMVVMCommands.cpp
+++ b/Engine/Plugins/Experimental/ECABridge/Source/ECABridge/Private/Commands/ECAMVVMCommands.cpp
@@ -1436,6 +1436,20 @@ FECACommandResult FECACommand_SetViewmodelBindingArgument::Execute(const TShared
 	// Type-filtered lookups return a subset and would cause arg_index to select the wrong pin
 	// for heterogeneous-input functions (matching the behavior in AddViewmodelBinding).
 	TArray<FMVVMBlueprintPinId> InputPins = MVVMSubsystem->GetConversionFunctionPins(WidgetBP, ConvValue, nullptr);
+	if (InputPins.Num() == 0 && ResolvedConvFunc)
+	{
+		// Fallback: construct pin IDs from the UFunction's parameter FProperties.
+		// GetConversionFunctionPins can return empty for BP Function Library UFunctions
+		// whose TryGetArgumentsForConversionFunction fails (matching AddViewmodelBinding).
+		for (TFieldIterator<FProperty> ParamIt(ResolvedConvFunc); ParamIt && ParamIt->HasAnyPropertyFlags(CPF_Parm); ++ParamIt)
+		{
+			if (!ParamIt->HasAnyPropertyFlags(CPF_ReturnParm))
+			{
+				TArray<FName> Names = { ParamIt->GetFName() };
+				InputPins.Add(FMVVMBlueprintPinId(MoveTemp(Names)));
+			}
+		}
+	}
 	if (!InputPins.IsValidIndex(ArgIndex))
 	{
 		return FECACommandResult::Error(FString::Printf(
